MotionDesign : ActorModifier
- Add SplinePath modifier that moves an actor along a path with various options (orient, scale, time, percentage, index, distance)

diff --git a/Engine/Plugins/Experimental/ActorModifier/Source/ActorModifierLayout/Private/Modifiers/ActorModifierLayoutSplinePathModifier.cpp b/Engine/Plugins/Experimental/ActorModifier/Source/ActorModifierLayout/Private/Modifiers/ActorModifierLayoutSplinePathModifier.cpp
new file mode 100644
--- /dev/null
+++ b/Engine/Plugins/Experimental/ActorModifier/Source/ActorModifierLayout/Private/Modifiers/ActorModifierLayoutSplinePathModifier.cpp
@@ -0,0 +1,324 @@
+// Copyright Epic Games, Inc. All Rights Reserved.
+
+#include "Modifiers/ActorModifierLayoutSplinePathModifier.h"
+
+#include "Components/SplineComponent.h"
+#include "Extensions/ActorModifierTransformUpdateExtension.h"
+#include "Shared/ActorModifierTransformShared.h"
+
+#define LOCTEXT_NAMESPACE "ActorModifierLayoutSplinePathModifier"
+
+void UActorModifierLayoutSplinePathModifier::SetSplineActorWeak(TWeakObjectPtr<AActor> InActor)
+{
+	if (SplineActorWeak == InActor)
+	{
+		return;
+	}
+
+	SplineActorWeak = InActor;
+	OnSplineActorWeakChanged();
+}
+
+void UActorModifierLayoutSplinePathModifier::SetSplineActor(AActor* InActor)
+{
+	SetSplineActorWeak(InActor);
+}
+
+void UActorModifierLayoutSplinePathModifier::SetSampleMode(EActorModifierLayoutSplinePathSampleMode InMode)
+{
+	if (SampleMode == InMode)
+	{
+		return;
+	}
+	
+	SampleMode = InMode;
+	OnSplineOptionsChanged();
+}
+
+void UActorModifierLayoutSplinePathModifier::SetProgress(float InProgress)
+{
+	InProgress = FMath::Max(InProgress, 0.f);
+	if (FMath::IsNearlyEqual(Progress, InProgress))
+	{
+		return;
+	}
+
+	Progress = InProgress;
+	OnSplineOptionsChanged();
+}
+
+void UActorModifierLayoutSplinePathModifier::SetDistance(float InDistance)
+{
+	InDistance = FMath::Max(InDistance, 0.f);
+	if (FMath::IsNearlyEqual(Distance, InDistance))
+	{
+		return;
+	}
+
+	Distance = InDistance;
+	OnSplineOptionsChanged();
+}
+
+void UActorModifierLayoutSplinePathModifier::SetTime(float InTime)
+{
+	InTime = FMath::Max(InTime, 0.f);
+	if (FMath::IsNearlyEqual(Time, InTime))
+	{
+		return;
+	}
+
+	Time = InTime;
+	OnSplineOptionsChanged();
+}
+
+void UActorModifierLayoutSplinePathModifier::SetPointIndex(int32 InIndex)
+{
+	InIndex = FMath::Max(InIndex, 0);
+	if (PointIndex == InIndex)
+	{
+		return;
+	}
+
+	PointIndex = InIndex;
+	OnSplineOptionsChanged();
+}
+
+void UActorModifierLayoutSplinePathModifier::SetOrient(bool bInOrient)
+{
+	if (bOrient == bInOrient)
+	{
+		return;
+	}
+
+	bOrient = bInOrient;
+	OnSplineOptionsChanged();
+}
+
+void UActorModifierLayoutSplinePathModifier::SetScale(bool bInScale)
+{
+	if (bScale == bInScale)
+	{
+		return;
+	}
+
+	bScale = bInScale;
+	OnSplineOptionsChanged();
+}
+
+#if WITH_EDITOR
+const TActorModifierPropertyChangeDispatcher<UActorModifierLayoutSplinePathModifier> UActorModifierLayoutSplinePathModifier::PropertyChangeDispatcher =
+{
+	{ GET_MEMBER_NAME_CHECKED(UActorModifierLayoutSplinePathModifier, SplineActorWeak), &UActorModifierLayoutSplinePathModifier::OnSplineActorWeakChanged },
+	{ GET_MEMBER_NAME_CHECKED(UActorModifierLayoutSplinePathModifier, SampleMode), &UActorModifierLayoutSplinePathModifier::OnSplineOptionsChanged },
+	{ GET_MEMBER_NAME_CHECKED(UActorModifierLayoutSplinePathModifier, Progress), &UActorModifierLayoutSplinePathModifier::OnSplineOptionsChanged },
+	{ GET_MEMBER_NAME_CHECKED(UActorModifierLayoutSplinePathModifier, Distance), &UActorModifierLayoutSplinePathModifier::OnSplineOptionsChanged },
+	{ GET_MEMBER_NAME_CHECKED(UActorModifierLayoutSplinePathModifier, Time), &UActorModifierLayoutSplinePathModifier::OnSplineOptionsChanged },
+	{ GET_MEMBER_NAME_CHECKED(UActorModifierLayoutSplinePathModifier, PointIndex), &UActorModifierLayoutSplinePathModifier::OnSplineOptionsChanged },
+	{ GET_MEMBER_NAME_CHECKED(UActorModifierLayoutSplinePathModifier, bOrient), &UActorModifierLayoutSplinePathModifier::OnSplineOptionsChanged },
+	{ GET_MEMBER_NAME_CHECKED(UActorModifierLayoutSplinePathModifier, bScale), &UActorModifierLayoutSplinePathModifier::OnSplineOptionsChanged },
+};
+
+void UActorModifierLayoutSplinePathModifier::PostEditChangeProperty(FPropertyChangedEvent& InPropertyChangedEvent)
+{
+	Super::PostEditChangeProperty(InPropertyChangedEvent);
+
+	PropertyChangeDispatcher.OnPropertyChanged(this, InPropertyChangedEvent);
+}
+#endif
+
+void UActorModifierLayoutSplinePathModifier::OnRenderStateUpdated(AActor* InActor, UActorComponent* InComponent)
+{
+	if (InActor == SplineActorWeak.Get())
+	{
+		OnSplineOptionsChanged();
+	}
+}
+
+void UActorModifierLayoutSplinePathModifier::OnTransformUpdated(AActor* InActor, bool bInParentMoved)
+{
+	if (InActor == SplineActorWeak.Get())
+	{
+		MarkModifierDirty();
+	}
+}
+
+void UActorModifierLayoutSplinePathModifier::OnModifierCDOSetup(FActorModifierCoreMetadata& InMetadata)
+{
+	Super::OnModifierCDOSetup(InMetadata);
+
+	InMetadata.SetName(TEXT("SplinePath"));
+	InMetadata.SetCategory(TEXT("Layout"));
+#if WITH_EDITOR
+	InMetadata.SetDescription(LOCTEXT("ModifierDescription", "Samples a spline for an actor to follow a path"));
+#endif
+
+	InMetadata.SetCompatibilityRule([](const AActor* InActor)->bool
+	{
+		return IsValid(InActor) && IsValid(InActor->GetRootComponent());
+	});
+}
+
+void UActorModifierLayoutSplinePathModifier::OnModifierAdded(EActorModifierCoreEnableReason InReason)
+{
+	Super::OnModifierAdded(InReason);
+	
+	AddExtension<FActorModifierRenderStateUpdateExtension>(this);
+	AddExtension<FActorModifierTransformUpdateExtension>(this);
+}
+
+void UActorModifierLayoutSplinePathModifier::OnModifierDisabled(EActorModifierCoreDisableReason InReason)
+{
+	Super::OnModifierDisabled(InReason);
+	
+	if (UActorModifierTransformShared* LayoutShared = GetShared<UActorModifierTransformShared>(/** CreateIfNone */false))
+	{
+		LayoutShared->RestoreActorState(this, GetModifiedActor());
+	}
+}
+
+void UActorModifierLayoutSplinePathModifier::Apply()
+{
+	AActor* ActorModified = GetModifiedActor();
+	if (!ActorModified)
+	{
+		return;
+	}
+
+	UActorModifierTransformShared* LayoutShared = GetShared<UActorModifierTransformShared>(true);
+	if (!LayoutShared)
+	{
+		Fail(LOCTEXT("InvalidTransformSharedObject", "Invalid transform shared object"));
+		return;
+	}
+
+	const USplineComponent* SplineComponent = SplineComponentWeak.Get();
+	if (!SplineComponent)
+	{
+		Next();
+		return;
+	}
+
+	if (SplineComponent->GetSplineLength() <= 0.f
+		|| SplineComponent->GetNumberOfSplinePoints() <= 0
+		|| SplineComponent->Duration <= UE_SMALL_NUMBER)
+	{
+		Fail(LOCTEXT("InvalidSplineValues", "Invalid spline values"));
+		return;
+	}
+
+	float SampleDistance = 0.f;
+
+	if (SampleMode == EActorModifierLayoutSplinePathSampleMode::Distance)
+	{
+		SampleDistance = Distance;
+	}
+	else if (SampleMode == EActorModifierLayoutSplinePathSampleMode::Percentage)
+	{
+		SampleDistance = SplineComponent->GetSplineLength() * Progress / 100.f;
+	}
+	else if (SampleMode == EActorModifierLayoutSplinePathSampleMode::Time)
+	{
+		SampleDistance = SplineComponent->GetSplineLength() * Time / SplineComponent->Duration;
+	}
+	else if (SampleMode == EActorModifierLayoutSplinePathSampleMode::Point)
+	{
+		int32 SampleIndex = PointIndex;
+		
+		if (SplineComponent->IsClosedLoop())
+		{
+			SampleIndex = FMath::Modulo(SampleIndex, SplineComponent->GetNumberOfSplinePoints());
+		}
+		
+		if (SampleIndex < SplineComponent->GetNumberOfSplinePoints())
+		{
+			SampleDistance = SplineComponent->GetDistanceAlongSplineAtSplinePoint(SampleIndex);
+		}
+		else
+		{
+			SampleDistance = SplineComponent->GetSplineLength();
+		}
+	}
+
+	SampleDistance = FMath::Max(SampleDistance, 0.f);
+	
+	if (SplineComponent->IsClosedLoop())
+	{
+		SampleDistance = FMath::Fmod(SampleDistance, SplineComponent->GetSplineLength());
+	}
+
+	LayoutShared->SaveActorState(this, ActorModified, EActorModifierTransformSharedState::Location);
+
+	const FVector WorldLocation = SplineComponent->GetLocationAtDistanceAlongSpline(SampleDistance, ESplineCoordinateSpace::World);
+	ActorModified->SetActorLocation(WorldLocation);
+
+	if (bOrient)
+	{
+		LayoutShared->SaveActorState(this, ActorModified, EActorModifierTransformSharedState::Rotation);
+		const FRotator WorldRotation = SplineComponent->GetRotationAtDistanceAlongSpline(SampleDistance, ESplineCoordinateSpace::World);
+		ActorModified->SetActorRotation(WorldRotation);
+	}
+	else
+	{
+		LayoutShared->RestoreActorState(this, ActorModified, EActorModifierTransformSharedState::Rotation);
+	}
+
+	if (bScale)
+	{
+		LayoutShared->SaveActorState(this, ActorModified, EActorModifierTransformSharedState::Scale);
+		const FVector WorldScale = SplineComponent->GetComponentScale() * SplineComponent->GetScaleAtDistanceAlongSpline(SampleDistance);
+		ActorModified->SetActorScale3D(WorldScale);
+	}
+	else
+	{
+		LayoutShared->RestoreActorState(this, ActorModified, EActorModifierTransformSharedState::Scale);
+	}
+
+	Next();
+}
+
+void UActorModifierLayoutSplinePathModifier::OnSplineActorWeakChanged()
+{
+	const AActor* SplineActor = SplineActorWeak.Get();
+
+	if (!SplineActor)
+	{
+		return;
+	}
+
+	USplineComponent* SplineComponent = SplineActor->FindComponentByClass<USplineComponent>();
+
+	// Don't allow actors without spline component
+	if (!SplineComponent)
+	{
+		SplineActorWeak.Reset();
+	}
+
+	// Don't update if we already track this component
+	if (SplineComponentWeak == SplineComponent)
+	{
+		return;
+	}
+
+	if (FActorModifierTransformUpdateExtension* TransformExtension = GetExtension<FActorModifierTransformUpdateExtension>())
+	{
+		if (const USplineComponent* OldComponent = SplineComponentWeak.Get())
+		{
+			TransformExtension->UntrackActor(OldComponent->GetOwner());
+		}
+
+		if (SplineComponent)
+		{
+			TransformExtension->TrackActor(SplineComponent->GetOwner(), /** Reset */false);
+		}
+	}
+	
+	SplineComponentWeak = SplineComponent;
+	OnSplineOptionsChanged();
+}
+
+void UActorModifierLayoutSplinePathModifier::OnSplineOptionsChanged()
+{
+	MarkModifierDirty();
+}
+
+#undef LOCTEXT_NAMESPACE
diff --git a/Engine/Plugins/Experimental/ActorModifier/Source/ActorModifierLayout/Public/Modifiers/ActorModifierLayoutSplinePathModifier.h b/Engine/Plugins/Experimental/ActorModifier/Source/ActorModifierLayout/Public/Modifiers/ActorModifierLayoutSplinePathModifier.h
new file mode 100644
--- /dev/null
+++ b/Engine/Plugins/Experimental/ActorModifier/Source/ActorModifierLayout/Public/Modifiers/ActorModifierLayoutSplinePathModifier.h
@@ -0,0 +1,177 @@
+// Copyright Epic Games, Inc. All Rights Reserved.
+
+#pragma once
+
+#include "Extensions/ActorModifierRenderStateUpdateExtension.h"
+#include "Extensions/ActorModifierTransformUpdateExtension.h"
+#include "Modifiers/ActorModifierCoreBase.h"
+#include "Utilities/ActorModifierPropertyChangeDispatcher.h"
+#include "ActorModifierLayoutSplinePathModifier.generated.h"
+
+class USplineComponent;
+
+UENUM(BlueprintType)
+enum class EActorModifierLayoutSplinePathSampleMode : uint8
+{
+	Percentage,
+	Distance,
+	Time,
+	Point
+};
+
+/**
+ * This modifier allows to sample a spline and update the actor transform based on the spline path
+ */
+UCLASS(MinimalAPI, BlueprintType)
+class UActorModifierLayoutSplinePathModifier : public UActorModifierCoreBase
+	, public IActorModifierRenderStateUpdateHandler
+	, public IActorModifierTransformUpdateHandler
+{
+	GENERATED_BODY()
+
+public:
+	void SetSplineActorWeak(TWeakObjectPtr<AActor> InActor);
+	TWeakObjectPtr<AActor> GetSplineActorWeak() const
+	{
+		return SplineActorWeak;
+	}
+
+	UFUNCTION(BlueprintCallable, Category="Motion Design|Modifiers|SplinePath")
+	void SetSplineActor(AActor* InActor);
+
+	UFUNCTION(BlueprintPure, Category="Motion Design|Modifiers|SplinePath")
+	AActor* GetSplineActor() const
+	{
+		return SplineActorWeak.Get();
+	}
+
+	UFUNCTION(BlueprintCallable, Category="Motion Design|Modifiers|SplinePath")
+	void SetSampleMode(EActorModifierLayoutSplinePathSampleMode InMode);
+
+	UFUNCTION(BlueprintPure, Category="Motion Design|Modifiers|SplinePath")
+	EActorModifierLayoutSplinePathSampleMode GetSampleMode() const
+	{
+		return SampleMode;
+	}
+
+	UFUNCTION(BlueprintCallable, Category="Motion Design|Modifiers|SplinePath")
+	void SetProgress(float InProgress);
+
+	UFUNCTION(BlueprintPure, Category="Motion Design|Modifiers|SplinePath")
+	float GetProgress() const
+	{
+		return Progress;
+	}
+
+	UFUNCTION(BlueprintCallable, Category="Motion Design|Modifiers|SplinePath")
+	void SetDistance(float InDistance);
+
+	UFUNCTION(BlueprintPure, Category="Motion Design|Modifiers|SplinePath")
+	float GetDistance() const
+	{
+		return Distance;
+	}
+
+	UFUNCTION(BlueprintCallable, Category="Motion Design|Modifiers|SplinePath")
+	void SetTime(float InTime);
+
+	UFUNCTION(BlueprintPure, Category="Motion Design|Modifiers|SplinePath")
+	float GetTime() const
+	{
+		return Time;
+	}
+
+	UFUNCTION(BlueprintCallable, Category="Motion Design|Modifiers|SplinePath")
+	void SetPointIndex(int32 InIndex);
+
+	UFUNCTION(BlueprintPure, Category="Motion Design|Modifiers|SplinePath")
+	int32 GetPointIndex() const
+	{
+		return PointIndex;
+	}
+
+	UFUNCTION(BlueprintCallable, Category="Motion Design|Modifiers|SplinePath")
+	void SetOrient(bool bInOrient);
+
+	UFUNCTION(BlueprintPure, Category="Motion Design|Modifiers|SplinePath")
+	bool GetOrient() const
+	{
+		return bOrient;
+	}
+
+	UFUNCTION(BlueprintCallable, Category="Motion Design|Modifiers|SplinePath")
+	void SetScale(bool bInScale);
+	
+	UFUNCTION(BlueprintPure, Category="Motion Design|Modifiers|SplinePath")
+	bool GetScale() const
+	{
+		return bScale;
+	}
+	
+protected:
+	//~ Begin UObject
+#if WITH_EDITOR
+	virtual void PostEditChangeProperty(FPropertyChangedEvent& InPropertyChangedEvent) override;
+#endif
+	//~ End UObject
+
+	//~ Begin IActorModifierRenderStateUpdateHandler
+	virtual void OnRenderStateUpdated(AActor* InActor, UActorComponent* InComponent) override;
+	virtual void OnActorVisibilityChanged(AActor* InActor) override {}
+	//~ End IActorModifierRenderStateUpdateHandler
+
+	//~ Begin IActorModifierTransformUpdateHandler
+	virtual void OnTransformUpdated(AActor* InActor, bool bInParentMoved) override;
+	//~ End IActorModifierTransformUpdateHandler
+
+	//~ Begin UActorModifierCoreBase
+	virtual void OnModifierCDOSetup(FActorModifierCoreMetadata& InMetadata) override;
+	virtual void OnModifierAdded(EActorModifierCoreEnableReason InReason) override;
+	virtual void OnModifierDisabled(EActorModifierCoreDisableReason InReason) override;
+	virtual void Apply() override;
+	//~ End UActorModifierCoreBase
+
+	void OnSplineActorWeakChanged();
+	void OnSplineOptionsChanged();
+	
+	/** Spline actor to retrieve the USplineComponent from */
+	UPROPERTY(EditInstanceOnly, Setter, Getter, Category="SplinePath", meta=(AllowPrivateAccess="true"))
+	TWeakObjectPtr<AActor> SplineActorWeak;
+
+	/** How to sample the spline */
+	UPROPERTY(EditInstanceOnly, Setter, Getter, Category="SplinePath", meta=(AllowPrivateAccess="true"))
+	EActorModifierLayoutSplinePathSampleMode SampleMode = EActorModifierLayoutSplinePathSampleMode::Percentage;
+	
+	/** Percentage progress to sample the spline at */
+	UPROPERTY(EditInstanceOnly, Setter, Getter, Category="SplinePath", meta=(EditCondition="SampleMode == EActorModifierLayoutSplinePathSampleMode::Percentage", EditConditionHides, Units=Percent, ClampMin="0", UIMin="0", UIMax="100", AllowPrivateAccess="true"))
+	float Progress = 0.f;
+
+	/** Distance to sample the spline at */
+	UPROPERTY(EditInstanceOnly, Setter, Getter, Category="SplinePath", meta=(EditCondition="SampleMode == EActorModifierLayoutSplinePathSampleMode::Distance", EditConditionHides, ClampMin="0", Units=Centimeters, AllowPrivateAccess="true"))
+	float Distance = 0.f;
+
+	/** Time to sample the spline at */
+	UPROPERTY(EditInstanceOnly, Setter, Getter, Category="SplinePath", meta=(EditCondition="SampleMode == EActorModifierLayoutSplinePathSampleMode::Time", EditConditionHides, ClampMin="0", AllowPrivateAccess="true"))
+	float Time = 0.f;
+	
+	/** Point index to sample the spline at */
+	UPROPERTY(EditInstanceOnly, Setter, Getter, Category="SplinePath", meta=(EditCondition="SampleMode == EActorModifierLayoutSplinePathSampleMode::Point", EditConditionHides, ClampMin="0", AllowPrivateAccess="true"))
+	int32 PointIndex = 0;
+
+	/** Orient actor based on spline tangent */
+	UPROPERTY(EditInstanceOnly, Setter="SetOrient", Getter="GetOrient", Category="SplinePath", meta=(AllowPrivateAccess="true"))
+	bool bOrient = true;
+
+	/** Apply scale based on spline */
+	UPROPERTY(EditInstanceOnly, Setter="SetScale", Getter="GetScale", Category="SplinePath", meta=(AllowPrivateAccess="true"))
+	bool bScale = false;
+
+	/** Spline component found on the spline actor */
+	UPROPERTY()
+	TWeakObjectPtr<USplineComponent> SplineComponentWeak;
+
+private:
+#if WITH_EDITOR
+	static const TActorModifierPropertyChangeDispatcher<UActorModifierLayoutSplinePathModifier> PropertyChangeDispatcher;
+#endif
+};
\ No newline at end of file
