MotionDesign : ActorModifier
- Improved spline path modifier based on feedback (only show actors with spline component in dropdown, use a base rotation to orient actor)
- Reusing cloner effector editor node picker since DynamicMeshConverter modifier uses ClonerEffectorMeshBuilder and eventually that modifier will be moved to ActorModifier/ActorModifierGeometry module

diff --git a/Engine/Plugins/Experimental/ActorModifier/ActorModifier.uplugin b/Engine/Plugins/Experimental/ActorModifier/ActorModifier.uplugin
--- a/Engine/Plugins/Experimental/ActorModifier/ActorModifier.uplugin
+++ b/Engine/Plugins/Experimental/ActorModifier/ActorModifier.uplugin
@@ -26,13 +26,13 @@
 			"LoadingPhase": "Default"
 		},
 		{
-			"Name": "ActorModifierLayout",
-			"Type": "Runtime",
+			"Name": "ActorModifierEditor",
+			"Type": "Editor",
 			"LoadingPhase": "Default"
 		},
 		{
-			"Name": "ActorModifierEditor",
-			"Type": "Editor",
+			"Name": "ActorModifierLayout",
+			"Type": "Runtime",
 			"LoadingPhase": "Default"
 		}
 	],
@@ -40,6 +40,10 @@
 		{
 			"Name": "ActorModifierCore",
 			"Enabled": true
+		},
+		{
+			"Name": "ClonerEffector",
+			"Enabled": true
 		}
 	]
 }
\ No newline at end of file
diff --git a/Engine/Plugins/Experimental/ActorModifier/Config/BaseActorModifier.ini b/Engine/Plugins/Experimental/ActorModifier/Config/BaseActorModifier.ini
--- a/Engine/Plugins/Experimental/ActorModifier/Config/BaseActorModifier.ini
+++ b/Engine/Plugins/Experimental/ActorModifier/Config/BaseActorModifier.ini
@@ -30,3 +30,4 @@
 +ClassRedirects=(OldName="/Script/AvalancheModifiers.AvaJustifyModifier",NewName="/Script/ActorModifierLayout.ActorModifierJustifyModifier")
 +ClassRedirects=(OldName="/Script/AvalancheModifiers.AvaLookAtModifier",NewName="/Script/ActorModifierLayout.ActorModifierLookAtModifier")
 +ClassRedirects=(OldName="/Script/AvalancheModifiers.AvaRadialArrangeModifier",NewName="/Script/ActorModifierLayout.ActorModifierRadialArrangeModifier")
++ClassRedirects=(OldName="/Script/ActorModifierLayout.ActorModifierLayoutSplinePathModifier",NewName="/Script/ActorModifierLayout.ActorModifierSplinePathModifier")
diff --git a/Engine/Plugins/Experimental/ActorModifier/Source/ActorModifierEditor/ActorModifierEditor.Build.cs b/Engine/Plugins/Experimental/ActorModifier/Source/ActorModifierEditor/ActorModifierEditor.Build.cs
--- a/Engine/Plugins/Experimental/ActorModifier/Source/ActorModifierEditor/ActorModifierEditor.Build.cs
+++ b/Engine/Plugins/Experimental/ActorModifier/Source/ActorModifierEditor/ActorModifierEditor.Build.cs
@@ -20,6 +20,8 @@ public class ActorModifierEditor : ModuleRules
             {
 	            "ActorModifierCore",
 	            "ActorModifier",
+				"ActorModifierLayout",
+				"ClonerEffectorEditor",
                 "CoreUObject",
                 "Engine",
                 "Projects",
diff --git a/Engine/Plugins/Experimental/ActorModifier/Source/ActorModifierEditor/Private/ActorModifierEditorModule.cpp b/Engine/Plugins/Experimental/ActorModifier/Source/ActorModifierEditor/Private/ActorModifierEditorModule.cpp
--- a/Engine/Plugins/Experimental/ActorModifier/Source/ActorModifierEditor/Private/ActorModifierEditorModule.cpp
+++ b/Engine/Plugins/Experimental/ActorModifier/Source/ActorModifierEditor/Private/ActorModifierEditorModule.cpp
@@ -4,11 +4,11 @@
 
 #include "ActorModifierTypes.h"
 #include "Customizations/ActorModifierEditorAnchorAlignmentPropertyTypeCustomization.h"
+#include "Customizations/ActorModifierEditorSplinePathModifierDetailCustomization.h"
+#include "Modifiers/ActorModifierSplinePathModifier.h"
 #include "PropertyEditorModule.h"
 #include "Styles/ActorModifierEditorStyle.h"
 
-#define LOCTEXT_NAMESPACE "ActorModifierEditorModule"
-
 IMPLEMENT_MODULE(FActorModifierEditorModule, ActorModifierEditor)
 
 void FActorModifierEditorModule::StartupModule()
@@ -19,6 +19,8 @@ void FActorModifierEditorModule::StartupModule()
 
 	PropertyModule.RegisterCustomPropertyTypeLayout(FActorModifierAnchorAlignment::StaticStruct()->GetFName()
 	, FOnGetPropertyTypeCustomizationInstance::CreateStatic(&FActorModifierEditorAnchorAlignmentPropertyTypeCustomization::MakeInstance));
+
+	PropertyModule.RegisterCustomClassLayout(UActorModifierSplinePathModifier::StaticClass()->GetFName(), FOnGetDetailCustomizationInstance::CreateStatic(&FActorModifierEditorSplinePathModifierDetailCustomization::MakeInstance));
 }
 
 void FActorModifierEditorModule::ShutdownModule()
@@ -26,8 +28,7 @@ void FActorModifierEditorModule::ShutdownModule()
 	if (FPropertyEditorModule* PropertyModule = FModuleManager::GetModulePtr<FPropertyEditorModule>(TEXT("PropertyEditor")))
 	{
 		PropertyModule->UnregisterCustomPropertyTypeLayout(FActorModifierAnchorAlignment::StaticStruct()->GetFName());
+		PropertyModule->UnregisterCustomClassLayout(UActorModifierSplinePathModifier::StaticClass()->GetFName());
 	}
 }
-
-#undef LOCTEXT_NAMESPACE
     
diff --git a/Engine/Plugins/Experimental/ActorModifier/Source/ActorModifierEditor/Private/Customizations/ActorModifierEditorSplinePathModifierDetailCustomization.cpp b/Engine/Plugins/Experimental/ActorModifier/Source/ActorModifierEditor/Private/Customizations/ActorModifierEditorSplinePathModifierDetailCustomization.cpp
new file mode 100644
--- /dev/null
+++ b/Engine/Plugins/Experimental/ActorModifier/Source/ActorModifierEditor/Private/Customizations/ActorModifierEditorSplinePathModifierDetailCustomization.cpp
@@ -0,0 +1,38 @@
+// Copyright Epic Games, Inc. All Rights Reserved.
+
+#include "Customizations/ActorModifierEditorSplinePathModifierDetailCustomization.h"
+
+#include "Cloner/Customizations/CEEditorClonerCustomActorPickerNodeBuilder.h"
+#include "Components/SplineComponent.h"
+#include "DetailCategoryBuilder.h"
+#include "DetailLayoutBuilder.h"
+#include "Modifiers/ActorModifierSplinePathModifier.h"
+
+void FActorModifierEditorSplinePathModifierDetailCustomization::CustomizeDetails(IDetailLayoutBuilder& InDetailBuilder)
+{
+	TSharedRef<IPropertyHandle> SplinePropertyHandle = InDetailBuilder.GetProperty(UActorModifierSplinePathModifier::GetSplineActorWeakPropertyName(), UActorModifierSplinePathModifier::StaticClass());
+
+	if (!SplinePropertyHandle->IsValidHandle())
+	{
+		return;
+	}
+
+	SplinePropertyHandle->MarkHiddenByCustomization();
+
+	IDetailCategoryBuilder& SplineCategoryBuilder = InDetailBuilder.EditCategory(SplinePropertyHandle->GetDefaultCategoryName(), SplinePropertyHandle->GetDefaultCategoryText());
+
+	SplineCategoryBuilder.AddCustomBuilder(MakeShared<FCEEditorClonerCustomActorPickerNodeBuilder>(
+		SplinePropertyHandle
+		, FOnShouldFilterActor::CreateStatic(&FActorModifierEditorSplinePathModifierDetailCustomization::OnFilterSplineActor))
+	);
+}
+
+bool FActorModifierEditorSplinePathModifierDetailCustomization::OnFilterSplineActor(const AActor* InActor)
+{
+	if (!IsValid(InActor))
+	{
+		return false;
+	}
+
+	return !!InActor->FindComponentByClass<USplineComponent>();
+}
diff --git a/Engine/Plugins/Experimental/ActorModifier/Source/ActorModifierEditor/Private/Customizations/ActorModifierEditorSplinePathModifierDetailCustomization.h b/Engine/Plugins/Experimental/ActorModifier/Source/ActorModifierEditor/Private/Customizations/ActorModifierEditorSplinePathModifierDetailCustomization.h
new file mode 100644
--- /dev/null
+++ b/Engine/Plugins/Experimental/ActorModifier/Source/ActorModifierEditor/Private/Customizations/ActorModifierEditorSplinePathModifierDetailCustomization.h
@@ -0,0 +1,26 @@
+// Copyright Epic Games, Inc. All Rights Reserved.
+
+#pragma once
+
+#include "IDetailCustomization.h"
+#include "Templates/SharedPointer.h"
+
+class AActor;
+class IPropertyHandle;
+
+/** Used to customize spline path modifier properties in details panel */
+class FActorModifierEditorSplinePathModifierDetailCustomization : public IDetailCustomization
+{
+public:
+	static TSharedRef<IDetailCustomization> MakeInstance()
+	{
+		return MakeShared<FActorModifierEditorSplinePathModifierDetailCustomization>();
+	}
+
+	//~ Begin IDetailCustomization
+	virtual void CustomizeDetails(IDetailLayoutBuilder& InDetailBuilder) override;
+	//~ End IDetailCustomization
+
+private:
+	static bool OnFilterSplineActor(const AActor* InActor);
+};
diff --git a/Engine/Plugins/Experimental/ActorModifier/Source/ActorModifierLayout/Private/Modifiers/ActorModifierLayoutSplinePathModifier.cpp b/Engine/Plugins/Experimental/ActorModifier/Source/ActorModifierLayout/Private/Modifiers/ActorModifierLayoutSplinePathModifier.cpp
deleted file mode 100644
diff --git a/Engine/Plugins/Experimental/ActorModifier/Source/ActorModifierLayout/Private/Modifiers/ActorModifierSplinePathModifier.cpp b/Engine/Plugins/Experimental/ActorModifier/Source/ActorModifierLayout/Private/Modifiers/ActorModifierSplinePathModifier.cpp
new file mode 100644
--- /dev/null
+++ b/Engine/Plugins/Experimental/ActorModifier/Source/ActorModifierLayout/Private/Modifiers/ActorModifierSplinePathModifier.cpp
@@ -0,0 +1,341 @@
+// Copyright Epic Games, Inc. All Rights Reserved.
+
+#include "Modifiers/ActorModifierSplinePathModifier.h"
+
+#include "Components/SplineComponent.h"
+#include "Extensions/ActorModifierTransformUpdateExtension.h"
+#include "Shared/ActorModifierTransformShared.h"
+
+#define LOCTEXT_NAMESPACE "ActorModifierLayoutSplinePathModifier"
+
+void UActorModifierSplinePathModifier::SetSplineActorWeak(TWeakObjectPtr<AActor> InActor)
+{
+	if (SplineActorWeak == InActor)
+	{
+		return;
+	}
+
+	SplineActorWeak = InActor;
+	OnSplineActorWeakChanged();
+}
+
+void UActorModifierSplinePathModifier::SetSplineActor(AActor* InActor)
+{
+	SetSplineActorWeak(InActor);
+}
+
+void UActorModifierSplinePathModifier::SetSampleMode(EActorModifierLayoutSplinePathSampleMode InMode)
+{
+	if (SampleMode == InMode)
+	{
+		return;
+	}
+	
+	SampleMode = InMode;
+	OnSplineOptionsChanged();
+}
+
+void UActorModifierSplinePathModifier::SetProgress(float InProgress)
+{
+	InProgress = FMath::Max(InProgress, 0.f);
+	if (FMath::IsNearlyEqual(Progress, InProgress))
+	{
+		return;
+	}
+
+	Progress = InProgress;
+	OnSplineOptionsChanged();
+}
+
+void UActorModifierSplinePathModifier::SetDistance(float InDistance)
+{
+	InDistance = FMath::Max(InDistance, 0.f);
+	if (FMath::IsNearlyEqual(Distance, InDistance))
+	{
+		return;
+	}
+
+	Distance = InDistance;
+	OnSplineOptionsChanged();
+}
+
+void UActorModifierSplinePathModifier::SetTime(float InTime)
+{
+	InTime = FMath::Max(InTime, 0.f);
+	if (FMath::IsNearlyEqual(Time, InTime))
+	{
+		return;
+	}
+
+	Time = InTime;
+	OnSplineOptionsChanged();
+}
+
+void UActorModifierSplinePathModifier::SetPointIndex(int32 InIndex)
+{
+	InIndex = FMath::Max(InIndex, 0);
+	if (PointIndex == InIndex)
+	{
+		return;
+	}
+
+	PointIndex = InIndex;
+	OnSplineOptionsChanged();
+}
+
+void UActorModifierSplinePathModifier::SetOrient(bool bInOrient)
+{
+	if (bOrient == bInOrient)
+	{
+		return;
+	}
+
+	bOrient = bInOrient;
+	OnSplineOptionsChanged();
+}
+
+void UActorModifierSplinePathModifier::SetBaseOrientation(const FRotator& InOrientation)
+{
+	if (BaseOrientation.Equals(InOrientation))
+	{
+		return;
+	}
+
+	BaseOrientation = InOrientation;
+	OnSplineOptionsChanged();
+}
+
+void UActorModifierSplinePathModifier::SetScale(bool bInScale)
+{
+	if (bScale == bInScale)
+	{
+		return;
+	}
+
+	bScale = bInScale;
+	OnSplineOptionsChanged();
+}
+
+#if WITH_EDITOR
+FName UActorModifierSplinePathModifier::GetSplineActorWeakPropertyName()
+{
+	return GET_MEMBER_NAME_CHECKED(UActorModifierSplinePathModifier, SplineActorWeak);
+}
+
+const TActorModifierPropertyChangeDispatcher<UActorModifierSplinePathModifier> UActorModifierSplinePathModifier::PropertyChangeDispatcher =
+{
+	{ GET_MEMBER_NAME_CHECKED(UActorModifierSplinePathModifier, SplineActorWeak), &UActorModifierSplinePathModifier::OnSplineActorWeakChanged },
+	{ GET_MEMBER_NAME_CHECKED(UActorModifierSplinePathModifier, SampleMode), &UActorModifierSplinePathModifier::OnSplineOptionsChanged },
+	{ GET_MEMBER_NAME_CHECKED(UActorModifierSplinePathModifier, Progress), &UActorModifierSplinePathModifier::OnSplineOptionsChanged },
+	{ GET_MEMBER_NAME_CHECKED(UActorModifierSplinePathModifier, Distance), &UActorModifierSplinePathModifier::OnSplineOptionsChanged },
+	{ GET_MEMBER_NAME_CHECKED(UActorModifierSplinePathModifier, Time), &UActorModifierSplinePathModifier::OnSplineOptionsChanged },
+	{ GET_MEMBER_NAME_CHECKED(UActorModifierSplinePathModifier, PointIndex), &UActorModifierSplinePathModifier::OnSplineOptionsChanged },
+	{ GET_MEMBER_NAME_CHECKED(UActorModifierSplinePathModifier, bOrient), &UActorModifierSplinePathModifier::OnSplineOptionsChanged },
+	{ GET_MEMBER_NAME_CHECKED(UActorModifierSplinePathModifier, BaseOrientation), &UActorModifierSplinePathModifier::OnSplineOptionsChanged },
+	{ GET_MEMBER_NAME_CHECKED(UActorModifierSplinePathModifier, bScale), &UActorModifierSplinePathModifier::OnSplineOptionsChanged },
+};
+
+void UActorModifierSplinePathModifier::PostEditChangeProperty(FPropertyChangedEvent& InPropertyChangedEvent)
+{
+	Super::PostEditChangeProperty(InPropertyChangedEvent);
+
+	PropertyChangeDispatcher.OnPropertyChanged(this, InPropertyChangedEvent);
+}
+#endif
+
+void UActorModifierSplinePathModifier::OnRenderStateUpdated(AActor* InActor, UActorComponent* InComponent)
+{
+	if (InActor == SplineActorWeak.Get())
+	{
+		OnSplineOptionsChanged();
+	}
+}
+
+void UActorModifierSplinePathModifier::OnTransformUpdated(AActor* InActor, bool bInParentMoved)
+{
+	if (InActor == SplineActorWeak.Get())
+	{
+		MarkModifierDirty();
+	}
+}
+
+void UActorModifierSplinePathModifier::OnModifierCDOSetup(FActorModifierCoreMetadata& InMetadata)
+{
+	Super::OnModifierCDOSetup(InMetadata);
+
+	InMetadata.SetName(TEXT("SplinePath"));
+	InMetadata.SetCategory(TEXT("Layout"));
+#if WITH_EDITOR
+	InMetadata.SetDescription(LOCTEXT("ModifierDescription", "Samples a spline for an actor to follow a path"));
+#endif
+
+	InMetadata.SetCompatibilityRule([](const AActor* InActor)->bool
+	{
+		return IsValid(InActor) && IsValid(InActor->GetRootComponent());
+	});
+}
+
+void UActorModifierSplinePathModifier::OnModifierAdded(EActorModifierCoreEnableReason InReason)
+{
+	Super::OnModifierAdded(InReason);
+	
+	AddExtension<FActorModifierRenderStateUpdateExtension>(this);
+	AddExtension<FActorModifierTransformUpdateExtension>(this);
+}
+
+void UActorModifierSplinePathModifier::OnModifierDisabled(EActorModifierCoreDisableReason InReason)
+{
+	Super::OnModifierDisabled(InReason);
+	
+	if (UActorModifierTransformShared* LayoutShared = GetShared<UActorModifierTransformShared>(/** CreateIfNone */false))
+	{
+		LayoutShared->RestoreActorState(this, GetModifiedActor());
+	}
+}
+
+void UActorModifierSplinePathModifier::Apply()
+{
+	AActor* ActorModified = GetModifiedActor();
+	if (!ActorModified)
+	{
+		return;
+	}
+
+	UActorModifierTransformShared* LayoutShared = GetShared<UActorModifierTransformShared>(true);
+	if (!LayoutShared)
+	{
+		Fail(LOCTEXT("InvalidTransformSharedObject", "Invalid transform shared object"));
+		return;
+	}
+
+	const USplineComponent* SplineComponent = SplineComponentWeak.Get();
+	if (!SplineComponent)
+	{
+		Next();
+		return;
+	}
+
+	if (SplineComponent->GetSplineLength() <= 0.f
+		|| SplineComponent->GetNumberOfSplinePoints() <= 0
+		|| SplineComponent->Duration <= UE_SMALL_NUMBER)
+	{
+		Fail(LOCTEXT("InvalidSplineValues", "Invalid spline values"));
+		return;
+	}
+
+	float SampleDistance = 0.f;
+
+	if (SampleMode == EActorModifierLayoutSplinePathSampleMode::Distance)
+	{
+		SampleDistance = Distance;
+	}
+	else if (SampleMode == EActorModifierLayoutSplinePathSampleMode::Percentage)
+	{
+		SampleDistance = SplineComponent->GetSplineLength() * Progress / 100.f;
+	}
+	else if (SampleMode == EActorModifierLayoutSplinePathSampleMode::Time)
+	{
+		SampleDistance = SplineComponent->GetSplineLength() * Time / SplineComponent->Duration;
+	}
+	else if (SampleMode == EActorModifierLayoutSplinePathSampleMode::Point)
+	{
+		int32 SampleIndex = PointIndex;
+		
+		if (SplineComponent->IsClosedLoop())
+		{
+			SampleIndex = FMath::Modulo(SampleIndex, SplineComponent->GetNumberOfSplinePoints());
+		}
+		
+		if (SampleIndex < SplineComponent->GetNumberOfSplinePoints())
+		{
+			SampleDistance = SplineComponent->GetDistanceAlongSplineAtSplinePoint(SampleIndex);
+		}
+		else
+		{
+			SampleDistance = SplineComponent->GetSplineLength();
+		}
+	}
+
+	SampleDistance = FMath::Max(SampleDistance, 0.f);
+	
+	if (SplineComponent->IsClosedLoop())
+	{
+		SampleDistance = FMath::Fmod(SampleDistance, SplineComponent->GetSplineLength());
+	}
+
+	LayoutShared->SaveActorState(this, ActorModified, EActorModifierTransformSharedState::Location);
+
+	const FVector WorldLocation = SplineComponent->GetLocationAtDistanceAlongSpline(SampleDistance, ESplineCoordinateSpace::World);
+	ActorModified->SetActorLocation(WorldLocation);
+
+	if (bOrient)
+	{
+		LayoutShared->SaveActorState(this, ActorModified, EActorModifierTransformSharedState::Rotation);
+		const FRotator WorldRotation = BaseOrientation + SplineComponent->GetRotationAtDistanceAlongSpline(SampleDistance, ESplineCoordinateSpace::World);
+		ActorModified->SetActorRotation(WorldRotation);
+	}
+	else
+	{
+		LayoutShared->RestoreActorState(this, ActorModified, EActorModifierTransformSharedState::Rotation);
+	}
+
+	if (bScale)
+	{
+		LayoutShared->SaveActorState(this, ActorModified, EActorModifierTransformSharedState::Scale);
+		const FVector WorldScale = SplineComponent->GetComponentScale() * SplineComponent->GetScaleAtDistanceAlongSpline(SampleDistance);
+		ActorModified->SetActorScale3D(WorldScale);
+	}
+	else
+	{
+		LayoutShared->RestoreActorState(this, ActorModified, EActorModifierTransformSharedState::Scale);
+	}
+
+	Next();
+}
+
+void UActorModifierSplinePathModifier::OnSplineActorWeakChanged()
+{
+	const AActor* SplineActor = SplineActorWeak.Get();
+
+	if (!SplineActor)
+	{
+		return;
+	}
+
+	USplineComponent* SplineComponent = SplineActor->FindComponentByClass<USplineComponent>();
+
+	// Don't allow actors without spline component
+	if (!SplineComponent)
+	{
+		SplineActorWeak.Reset();
+	}
+
+	// Don't update if we already track this component
+	if (SplineComponentWeak == SplineComponent)
+	{
+		return;
+	}
+
+	if (FActorModifierTransformUpdateExtension* TransformExtension = GetExtension<FActorModifierTransformUpdateExtension>())
+	{
+		if (const USplineComponent* OldComponent = SplineComponentWeak.Get())
+		{
+			TransformExtension->UntrackActor(OldComponent->GetOwner());
+		}
+
+		if (SplineComponent)
+		{
+			TransformExtension->TrackActor(SplineComponent->GetOwner(), /** Reset */false);
+		}
+	}
+	
+	SplineComponentWeak = SplineComponent;
+	OnSplineOptionsChanged();
+}
+
+void UActorModifierSplinePathModifier::OnSplineOptionsChanged()
+{
+	MarkModifierDirty();
+}
+
+#undef LOCTEXT_NAMESPACE
diff --git a/Engine/Plugins/Experimental/ActorModifier/Source/ActorModifierLayout/Public/Modifiers/ActorModifierLayoutSplinePathModifier.h b/Engine/Plugins/Experimental/ActorModifier/Source/ActorModifierLayout/Public/Modifiers/ActorModifierLayoutSplinePathModifier.h
deleted file mode 100644
diff --git a/Engine/Plugins/Experimental/ActorModifier/Source/ActorModifierLayout/Public/Modifiers/ActorModifierSplinePathModifier.h b/Engine/Plugins/Experimental/ActorModifier/Source/ActorModifierLayout/Public/Modifiers/ActorModifierSplinePathModifier.h
new file mode 100644
--- /dev/null
+++ b/Engine/Plugins/Experimental/ActorModifier/Source/ActorModifierLayout/Public/Modifiers/ActorModifierSplinePathModifier.h
@@ -0,0 +1,194 @@
+// Copyright Epic Games, Inc. All Rights Reserved.
+
+#pragma once
+
+#include "Extensions/ActorModifierRenderStateUpdateExtension.h"
+#include "Extensions/ActorModifierTransformUpdateExtension.h"
+#include "Modifiers/ActorModifierCoreBase.h"
+#include "Utilities/ActorModifierPropertyChangeDispatcher.h"
+#include "ActorModifierSplinePathModifier.generated.h"
+
+class USplineComponent;
+
+UENUM(BlueprintType)
+enum class EActorModifierLayoutSplinePathSampleMode : uint8
+{
+	Percentage,
+	Distance,
+	Time,
+	Point
+};
+
+/**
+ * This modifier allows to sample a spline and update the actor transform based on the spline path
+ */
+UCLASS(MinimalAPI, BlueprintType)
+class UActorModifierSplinePathModifier : public UActorModifierCoreBase
+	, public IActorModifierRenderStateUpdateHandler
+	, public IActorModifierTransformUpdateHandler
+{
+	GENERATED_BODY()
+
+public:
+#if WITH_EDITOR
+	ACTORMODIFIERLAYOUT_API static FName GetSplineActorWeakPropertyName();
+#endif
+	
+	void SetSplineActorWeak(TWeakObjectPtr<AActor> InActor);
+	TWeakObjectPtr<AActor> GetSplineActorWeak() const
+	{
+		return SplineActorWeak;
+	}
+
+	UFUNCTION(BlueprintCallable, Category="Motion Design|Modifiers|SplinePath")
+	void SetSplineActor(AActor* InActor);
+
+	UFUNCTION(BlueprintPure, Category="Motion Design|Modifiers|SplinePath")
+	AActor* GetSplineActor() const
+	{
+		return SplineActorWeak.Get();
+	}
+
+	UFUNCTION(BlueprintCallable, Category="Motion Design|Modifiers|SplinePath")
+	void SetSampleMode(EActorModifierLayoutSplinePathSampleMode InMode);
+
+	UFUNCTION(BlueprintPure, Category="Motion Design|Modifiers|SplinePath")
+	EActorModifierLayoutSplinePathSampleMode GetSampleMode() const
+	{
+		return SampleMode;
+	}
+
+	UFUNCTION(BlueprintCallable, Category="Motion Design|Modifiers|SplinePath")
+	void SetProgress(float InProgress);
+
+	UFUNCTION(BlueprintPure, Category="Motion Design|Modifiers|SplinePath")
+	float GetProgress() const
+	{
+		return Progress;
+	}
+
+	UFUNCTION(BlueprintCallable, Category="Motion Design|Modifiers|SplinePath")
+	void SetDistance(float InDistance);
+
+	UFUNCTION(BlueprintPure, Category="Motion Design|Modifiers|SplinePath")
+	float GetDistance() const
+	{
+		return Distance;
+	}
+
+	UFUNCTION(BlueprintCallable, Category="Motion Design|Modifiers|SplinePath")
+	void SetTime(float InTime);
+
+	UFUNCTION(BlueprintPure, Category="Motion Design|Modifiers|SplinePath")
+	float GetTime() const
+	{
+		return Time;
+	}
+
+	UFUNCTION(BlueprintCallable, Category="Motion Design|Modifiers|SplinePath")
+	void SetPointIndex(int32 InIndex);
+
+	UFUNCTION(BlueprintPure, Category="Motion Design|Modifiers|SplinePath")
+	int32 GetPointIndex() const
+	{
+		return PointIndex;
+	}
+
+	UFUNCTION(BlueprintCallable, Category="Motion Design|Modifiers|SplinePath")
+	void SetOrient(bool bInOrient);
+
+	UFUNCTION(BlueprintPure, Category="Motion Design|Modifiers|SplinePath")
+	bool GetOrient() const
+	{
+		return bOrient;
+	}
+
+	UFUNCTION(BlueprintCallable, Category="Motion Design|Modifiers|SplinePath")
+	void SetBaseOrientation(const FRotator& InOrientation);
+
+	UFUNCTION(BlueprintPure, Category="Motion Design|Modifiers|SplinePath")
+	const FRotator& GetBaseOrientation() const
+	{
+		return BaseOrientation;
+	}
+
+	UFUNCTION(BlueprintCallable, Category="Motion Design|Modifiers|SplinePath")
+	void SetScale(bool bInScale);
+	
+	UFUNCTION(BlueprintPure, Category="Motion Design|Modifiers|SplinePath")
+	bool GetScale() const
+	{
+		return bScale;
+	}
+	
+protected:
+	//~ Begin UObject
+#if WITH_EDITOR
+	virtual void PostEditChangeProperty(FPropertyChangedEvent& InPropertyChangedEvent) override;
+#endif
+	//~ End UObject
+
+	//~ Begin IActorModifierRenderStateUpdateHandler
+	virtual void OnRenderStateUpdated(AActor* InActor, UActorComponent* InComponent) override;
+	virtual void OnActorVisibilityChanged(AActor* InActor) override {}
+	//~ End IActorModifierRenderStateUpdateHandler
+
+	//~ Begin IActorModifierTransformUpdateHandler
+	virtual void OnTransformUpdated(AActor* InActor, bool bInParentMoved) override;
+	//~ End IActorModifierTransformUpdateHandler
+
+	//~ Begin UActorModifierCoreBase
+	virtual void OnModifierCDOSetup(FActorModifierCoreMetadata& InMetadata) override;
+	virtual void OnModifierAdded(EActorModifierCoreEnableReason InReason) override;
+	virtual void OnModifierDisabled(EActorModifierCoreDisableReason InReason) override;
+	virtual void Apply() override;
+	//~ End UActorModifierCoreBase
+
+	void OnSplineActorWeakChanged();
+	void OnSplineOptionsChanged();
+	
+	/** Spline actor to retrieve the USplineComponent from */
+	UPROPERTY(EditInstanceOnly, Setter, Getter, Category="SplinePath", meta=(AllowPrivateAccess="true"))
+	TWeakObjectPtr<AActor> SplineActorWeak;
+
+	/** How to sample the spline */
+	UPROPERTY(EditInstanceOnly, Setter, Getter, Category="SplinePath", meta=(AllowPrivateAccess="true"))
+	EActorModifierLayoutSplinePathSampleMode SampleMode = EActorModifierLayoutSplinePathSampleMode::Percentage;
+	
+	/** Percentage progress to sample the spline at */
+	UPROPERTY(EditInstanceOnly, Setter, Getter, Category="SplinePath", meta=(EditCondition="SampleMode == EActorModifierLayoutSplinePathSampleMode::Percentage", EditConditionHides, Units=Percent, ClampMin="0", UIMin="0", UIMax="100", AllowPrivateAccess="true"))
+	float Progress = 0.f;
+
+	/** Distance to sample the spline at */
+	UPROPERTY(EditInstanceOnly, Setter, Getter, Category="SplinePath", meta=(EditCondition="SampleMode == EActorModifierLayoutSplinePathSampleMode::Distance", EditConditionHides, ClampMin="0", Units=Centimeters, AllowPrivateAccess="true"))
+	float Distance = 0.f;
+
+	/** Time to sample the spline at */
+	UPROPERTY(EditInstanceOnly, Setter, Getter, Category="SplinePath", meta=(EditCondition="SampleMode == EActorModifierLayoutSplinePathSampleMode::Time", EditConditionHides, ClampMin="0", AllowPrivateAccess="true"))
+	float Time = 0.f;
+	
+	/** Point index to sample the spline at */
+	UPROPERTY(EditInstanceOnly, Setter, Getter, Category="SplinePath", meta=(EditCondition="SampleMode == EActorModifierLayoutSplinePathSampleMode::Point", EditConditionHides, ClampMin="0", AllowPrivateAccess="true"))
+	int32 PointIndex = 0;
+
+	/** Orient actor based on spline tangent */
+	UPROPERTY(EditInstanceOnly, Setter="SetOrient", Getter="GetOrient", Category="SplinePath", meta=(AllowPrivateAccess="true"))
+	bool bOrient = true;
+
+	/** Base rotation added on top of the orientation rotation computed */
+	UPROPERTY(EditInstanceOnly, Setter, Getter, Category="SplinePath", meta=(EditCondition="bOrient", EditConditionHides, AllowPrivateAccess="true"))
+	FRotator BaseOrientation = FRotator::ZeroRotator;
+
+	/** Apply scale based on spline point scale */
+	UPROPERTY(EditInstanceOnly, Setter="SetScale", Getter="GetScale", Category="SplinePath", meta=(AllowPrivateAccess="true"))
+	bool bScale = false;
+
+	/** Spline component found on the spline actor */
+	UPROPERTY()
+	TWeakObjectPtr<USplineComponent> SplineComponentWeak;
+
+private:
+#if WITH_EDITOR
+	static const TActorModifierPropertyChangeDispatcher<UActorModifierSplinePathModifier> PropertyChangeDispatcher;
+#endif
+};
\ No newline at end of file
diff --git a/Engine/Plugins/Experimental/ClonerEffector/Source/ClonerEffectorEditor/ClonerEffectorEditor.Build.cs b/Engine/Plugins/Experimental/ClonerEffector/Source/ClonerEffectorEditor/ClonerEffectorEditor.Build.cs
--- a/Engine/Plugins/Experimental/ClonerEffector/Source/ClonerEffectorEditor/ClonerEffectorEditor.Build.cs
+++ b/Engine/Plugins/Experimental/ClonerEffector/Source/ClonerEffectorEditor/ClonerEffectorEditor.Build.cs
@@ -15,8 +15,7 @@ public class ClonerEffectorEditor : ModuleRules
 		        "CoreUObject",
 				"EditorSubsystem",
 		        "Engine",
-		        "Slate",
-		        "SlateCore",
+		        "PropertyEditor",
 		        "UnrealEd"
 	        }
         );
@@ -31,8 +30,9 @@ public class ClonerEffectorEditor : ModuleRules
 		        "NiagaraEditor",
 		        "NiagaraSimCaching",
 		        "Projects",
-				"PropertyEditor",
 				"Sequencer",
+				"Slate",
+				"SlateCore",
 				"ToolMenus"
 	        }
         );
diff --git a/Engine/Plugins/Experimental/ClonerEffector/Source/ClonerEffectorEditor/Private/Cloner/Customizations/CEEditorClonerCustomActorPickerNodeBuilder.h b/Engine/Plugins/Experimental/ClonerEffector/Source/ClonerEffectorEditor/Private/Cloner/Customizations/CEEditorClonerCustomActorPickerNodeBuilder.h
deleted file mode 100644
diff --git a/Engine/Plugins/Experimental/ClonerEffector/Source/ClonerEffectorEditor/Private/Cloner/Customizations/CEEditorClonerMeshLayoutDetailCustomization.cpp b/Engine/Plugins/Experimental/ClonerEffector/Source/ClonerEffectorEditor/Private/Cloner/Customizations/CEEditorClonerMeshLayoutDetailCustomization.cpp
--- a/Engine/Plugins/Experimental/ClonerEffector/Source/ClonerEffectorEditor/Private/Cloner/Customizations/CEEditorClonerMeshLayoutDetailCustomization.cpp
+++ b/Engine/Plugins/Experimental/ClonerEffector/Source/ClonerEffectorEditor/Private/Cloner/Customizations/CEEditorClonerMeshLayoutDetailCustomization.cpp
@@ -2,7 +2,7 @@
 
 #include "CEEditorClonerMeshLayoutDetailCustomization.h"
 
-#include "CEEditorClonerCustomActorPickerNodeBuilder.h"
+#include "Cloner/Customizations/CEEditorClonerCustomActorPickerNodeBuilder.h"
 #include "Cloner/Layouts/CEClonerMeshLayout.h"
 #include "Components/SkeletalMeshComponent.h"
 #include "Components/StaticMeshComponent.h"
diff --git a/Engine/Plugins/Experimental/ClonerEffector/Source/ClonerEffectorEditor/Private/Cloner/Customizations/CEEditorClonerSplineLayoutDetailCustomization.cpp b/Engine/Plugins/Experimental/ClonerEffector/Source/ClonerEffectorEditor/Private/Cloner/Customizations/CEEditorClonerSplineLayoutDetailCustomization.cpp
--- a/Engine/Plugins/Experimental/ClonerEffector/Source/ClonerEffectorEditor/Private/Cloner/Customizations/CEEditorClonerSplineLayoutDetailCustomization.cpp
+++ b/Engine/Plugins/Experimental/ClonerEffector/Source/ClonerEffectorEditor/Private/Cloner/Customizations/CEEditorClonerSplineLayoutDetailCustomization.cpp
@@ -2,7 +2,7 @@
 
 #include "Cloner/Customizations/CEEditorClonerSplineLayoutDetailCustomization.h"
 
-#include "CEEditorClonerCustomActorPickerNodeBuilder.h"
+#include "Cloner/Customizations/CEEditorClonerCustomActorPickerNodeBuilder.h"
 #include "Cloner/Layouts/CEClonerSplineLayout.h"
 #include "Components/SplineComponent.h"
 #include "DetailCategoryBuilder.h"
diff --git a/Engine/Plugins/Experimental/ClonerEffector/Source/ClonerEffectorEditor/Public/Cloner/Customizations/CEEditorClonerCustomActorPickerNodeBuilder.h b/Engine/Plugins/Experimental/ClonerEffector/Source/ClonerEffectorEditor/Public/Cloner/Customizations/CEEditorClonerCustomActorPickerNodeBuilder.h
new file mode 100644
--- /dev/null
+++ b/Engine/Plugins/Experimental/ClonerEffector/Source/ClonerEffectorEditor/Public/Cloner/Customizations/CEEditorClonerCustomActorPickerNodeBuilder.h
@@ -0,0 +1,48 @@
+// Copyright Epic Games, Inc. All Rights Reserved.
+
+#pragma once
+
+#include "ActorPickerMode.h"
+#include "IDetailCustomNodeBuilder.h"
+#include "Templates/SharedPointer.h"
+
+class IPropertyHandle;
+class IPropertyUtilities;
+class SComboButton;
+class SWidget;
+
+/** Used to customize an actor picker to filter items based on a delegate */
+class CLONEREFFECTOREDITOR_API FCEEditorClonerCustomActorPickerNodeBuilder : public IDetailCustomNodeBuilder, public TSharedFromThis<FCEEditorClonerCustomActorPickerNodeBuilder> 
+{
+public:
+	explicit FCEEditorClonerCustomActorPickerNodeBuilder(const TSharedRef<IPropertyHandle>& InPropertyHandle, FOnShouldFilterActor InActorFilterDelegate)
+		: PropertyHandle(InPropertyHandle)
+		, ActorFilterDelegate(InActorFilterDelegate)
+	{}
+
+private:
+	static TSharedPtr<SWidget> FindWidgetType(TSharedPtr<SWidget> InSearchWidget, FName InTypeName, bool bInReverseSearch);
+
+	//~ Begin IDetailCustomNodeBuilder
+	virtual FName GetName() const override;
+	virtual void GenerateHeaderRowContent(FDetailWidgetRow& InNodeRow) override;
+	virtual void GenerateChildContent(IDetailChildrenBuilder& InChildrenBuilder) override;
+	virtual TSharedPtr<IPropertyHandle> GetPropertyHandle() const override;
+	//~ End IDetailCustomNodeBuilder
+
+	FText GetPickerLabelText() const;
+	FText GetSelectTooltipText() const;
+	FText GetUseSelectTooltipText() const;
+	TSharedRef<SWidget> GetActorPickerWidget() const;
+
+	void OnActorSelected(AActor* InSelection) const;
+	void OnSelectActor() const;
+	void OnUseSelectedActor() const;
+
+	AActor* GetPropertyActor() const;
+	AActor* GetSelectedActor() const;
+
+	TSharedRef<IPropertyHandle> PropertyHandle;
+	TSharedPtr<SComboButton> ComboButton;
+	FOnShouldFilterActor ActorFilterDelegate;
+};
