MotionDesign : ActorModifier
- Improved LookAt modifier code to orient self to target actor, fixed X axis rotation was inverted
- Adapting test to match expected result

diff --git a/Engine/Plugins/Experimental/ActorModifier/Source/ActorModifier/Private/Utilities/ActorModifierActorUtils.cpp b/Engine/Plugins/Experimental/ActorModifier/Source/ActorModifier/Private/Utilities/ActorModifierActorUtils.cpp
--- a/Engine/Plugins/Experimental/ActorModifier/Source/ActorModifier/Private/Utilities/ActorModifierActorUtils.cpp
+++ b/Engine/Plugins/Experimental/ActorModifier/Source/ActorModifier/Private/Utilities/ActorModifierActorUtils.cpp
@@ -188,7 +188,7 @@ FOrientedBox UE::ActorModifier::ActorUtils::GetOrientedBox(const FBox& InLocalBo
 
 FRotator UE::ActorModifier::ActorUtils::FindLookAtRotation(const FVector& InEyePosition, const FVector& InTargetPosition, EActorModifierAxis InAxis, bool bInFlipAxis)
 {
-	const FVector Direction = bInFlipAxis ? (InTargetPosition - InEyePosition).GetSafeNormal() : (InEyePosition - InTargetPosition).GetSafeNormal();
+	const FVector Direction = bInFlipAxis ? (InEyePosition - InTargetPosition).GetSafeNormal() : (InTargetPosition - InEyePosition).GetSafeNormal();
 
 	if (Direction.IsNearlyZero())
 	{
@@ -205,11 +205,11 @@ FRotator UE::ActorModifier::ActorUtils::FindLookAtRotation(const FVector& InEyeP
 		break;
 
 	case EActorModifierAxis::Y:
-		AxisQuat = FQuat(FVector::ZAxisVector, FMath::DegreesToRadians(90.0f));
+		AxisQuat = FQuat(FVector::ZAxisVector, FMath::DegreesToRadians(-90.0f));
 		break;
 
 	case EActorModifierAxis::Z:
-		AxisQuat = FQuat(FVector::YAxisVector, FMath::DegreesToRadians(-90.0f));
+		AxisQuat = FQuat(FVector::YAxisVector, FMath::DegreesToRadians(90.0f));
 		break;
 
 	default:
diff --git a/Engine/Plugins/Experimental/ActorModifier/Source/ActorModifierLayout/Private/Modifiers/ActorModifierLookAtModifier.cpp b/Engine/Plugins/Experimental/ActorModifier/Source/ActorModifierLayout/Private/Modifiers/ActorModifierLookAtModifier.cpp
--- a/Engine/Plugins/Experimental/ActorModifier/Source/ActorModifierLayout/Private/Modifiers/ActorModifierLookAtModifier.cpp
+++ b/Engine/Plugins/Experimental/ActorModifier/Source/ActorModifierLayout/Private/Modifiers/ActorModifierLookAtModifier.cpp
@@ -6,6 +6,7 @@
 #include "GameFramework/Actor.h"
 #include "Misc/ITransaction.h"
 #include "Shared/ActorModifierTransformShared.h"
+#include "Utilities/ActorModifierActorUtils.h"
 
 #define LOCTEXT_NAMESPACE "ActorModifierLookAtModifier"
 
@@ -20,6 +21,11 @@ void UActorModifierLookAtModifier::OnModifierAdded(EActorModifierCoreEnableReaso
 		SceneExtension->TrackSceneTree(0, &ReferenceActor);
 	}
 
+	if (InReason == EActorModifierCoreEnableReason::User)
+	{
+		OrientationAxis = EActorModifierAxis::X;
+	}
+
 	bDeprecatedPropertiesMigrated = true;
 }
 
@@ -27,9 +33,9 @@ void UActorModifierLookAtModifier::OnModifierEnabled(EActorModifierCoreEnableRea
 {
 	Super::OnModifierEnabled(InReason);
 	
-	if (UActorModifierTransformShared* LayoutShared = GetShared<UActorModifierTransformShared>(true))
+	if (UActorModifierTransformShared* LayoutShared = GetShared<UActorModifierTransformShared>(/** CreateIfNone */true))
 	{
-		LayoutShared->SaveActorState(this, GetModifiedActor());
+		LayoutShared->SaveActorState(this, GetModifiedActor(), EActorModifierTransformSharedState::Rotation);
 	}
 }
 
@@ -37,9 +43,9 @@ void UActorModifierLookAtModifier::OnModifierDisabled(EActorModifierCoreDisableR
 {
 	Super::OnModifierDisabled(InReason);
 
-	if (UActorModifierTransformShared* LayoutShared = GetShared<UActorModifierTransformShared>(false))
+	if (UActorModifierTransformShared* LayoutShared = GetShared<UActorModifierTransformShared>(/** CreateIfNone */false))
 	{
-		LayoutShared->RestoreActorState(this, GetModifiedActor());
+		LayoutShared->RestoreActorState(this, GetModifiedActor(), EActorModifierTransformSharedState::Rotation);
 	}
 }
 
@@ -54,49 +60,8 @@ void UActorModifierLookAtModifier::Apply()
 		return;
 	}
 
-	auto LookAtDirection = [this, ModifyActor, LookActor]() -> FVector
-	{
-		FVector OutDirection = bFlipAxis
-			? LookActor->GetActorLocation() - ModifyActor->GetActorLocation()
-			: ModifyActor->GetActorLocation() - LookActor->GetActorLocation();
-
-		OutDirection.Normalize();
-		if (OutDirection.SizeSquared() < 0.5f)
-		{
-			// Assert possible if OutDirection is not normalized.
-			OutDirection = FVector(1, 0, 0);
-		}
-
-		return OutDirection;
-	};
-
-	auto NormalOffset = [this, ModifyActor](const FQuat& InNewRotation) -> FQuat
-	{
-		FQuat OutNormal = bFlipAxis
-			? InNewRotation - ModifyActor->GetActorRotation().Quaternion()
-			: ModifyActor->GetActorRotation().Quaternion() - InNewRotation;
-
-		OutNormal.Normalize();
-
-		return OutNormal;
-	};
-
-	FMatrix NewRotation = FRotationMatrix::Identity;
-		
-	switch (Axis)
-	{
-		case EActorModifierAlignment::Horizontal:
-			NewRotation = FRotationMatrix::MakeFromX(LookAtDirection());
-			break;
-		case EActorModifierAlignment::Vertical:
-			NewRotation = FRotationMatrix::MakeFromY(LookAtDirection());
-			break;
-		case EActorModifierAlignment::Depth:
-			NewRotation = FRotationMatrix::MakeFromZ(LookAtDirection());
-			break;
-	}
-
-	ModifyActor->SetActorRotation(NewRotation.ToQuat());
+	const FRotator NewRotation = UE::ActorModifier::ActorUtils::FindLookAtRotation(ModifyActor->GetActorLocation(), LookActor->GetActorLocation(), OrientationAxis, bFlipAxis);
+	ModifyActor->SetActorRotation(NewRotation);
 
 	Next();
 }
@@ -118,6 +83,15 @@ void UActorModifierLookAtModifier::PostLoad()
 
 		bDeprecatedPropertiesMigrated = true;
 	}
+
+	if (OrientationAxis == EActorModifierAxis::None)
+	{
+		PRAGMA_DISABLE_DEPRECATION_WARNINGS
+
+		OrientationAxis = static_cast<EActorModifierAxis>(1 << static_cast<int32>(Axis));
+
+		PRAGMA_ENABLE_DEPRECATION_WARNINGS
+	}
 	
 	Super::PostLoad();
 }
@@ -130,14 +104,14 @@ void UActorModifierLookAtModifier::PostEditChangeProperty(FPropertyChangedEvent&
 	const FName MemberName = PropertyChangedEvent.GetMemberPropertyName();
 	
 	static const FName ReferenceActorPropertyName = GET_MEMBER_NAME_CHECKED(UActorModifierLookAtModifier, ReferenceActor);
+	static const FName OrientationAxisPropertyName = GET_MEMBER_NAME_CHECKED(UActorModifierLookAtModifier, OrientationAxis);
 	static const FName FlipAxisPropertyName = GET_MEMBER_NAME_CHECKED(UActorModifierLookAtModifier, bFlipAxis);
-	static const FName AxisPropertyName = GET_MEMBER_NAME_CHECKED(UActorModifierLookAtModifier, Axis);
-	
+
 	if (MemberName == ReferenceActorPropertyName)
 	{
 		OnReferenceActorChanged();
 	}
-	else if (MemberName == AxisPropertyName
+	else if (MemberName == OrientationAxisPropertyName
 		|| MemberName == FlipAxisPropertyName)
 	{
 		MarkModifierDirty();
@@ -176,17 +150,25 @@ void UActorModifierLookAtModifier::OnReferenceActorChanged()
 	}
 }
 
-void UActorModifierLookAtModifier::SetAxis(const EActorModifierAlignment InAxis)
+void UActorModifierLookAtModifier::SetOrientationAxis(EActorModifierAxis InAxis)
 {
-	Axis = InAxis;
+	if (OrientationAxis == InAxis || OrientationAxis == EActorModifierAxis::None)
+	{
+		return;
+	}
 
+	OrientationAxis = InAxis;
 	MarkModifierDirty();
 }
 
 void UActorModifierLookAtModifier::SetFlipAxis(const bool bInFlipAxis)
 {
-	bFlipAxis = bInFlipAxis;
+	if (bFlipAxis == bInFlipAxis)
+	{
+		return;
+	}
 
+	bFlipAxis = bInFlipAxis;
 	MarkModifierDirty();
 }
 
diff --git a/Engine/Plugins/Experimental/ActorModifier/Source/ActorModifierLayout/Public/Modifiers/ActorModifierLookAtModifier.h b/Engine/Plugins/Experimental/ActorModifier/Source/ActorModifierLayout/Public/Modifiers/ActorModifierLookAtModifier.h
--- a/Engine/Plugins/Experimental/ActorModifier/Source/ActorModifierLayout/Public/Modifiers/ActorModifierLookAtModifier.h
+++ b/Engine/Plugins/Experimental/ActorModifier/Source/ActorModifierLayout/Public/Modifiers/ActorModifierLookAtModifier.h
@@ -27,16 +27,16 @@ public:
 	{
 		return ReferenceActor;
 	}
-
+	
 	/** Sets the axis that will point towards the reference actor. */
 	UFUNCTION(BlueprintCallable, Category="Motion Design|Modifiers|LookAt")
-	ACTORMODIFIERLAYOUT_API void SetAxis(const EActorModifierAlignment InAxis);
+	ACTORMODIFIERLAYOUT_API void SetOrientationAxis(EActorModifierAxis InAxis);
 
-	/** Returns the axis that will point towards t he reference actor. */
+	/** Returns the axis that will point towards the reference actor. */
 	UFUNCTION(BlueprintPure, Category="Motion Design|Modifiers|LookAt")
-	EActorModifierAlignment GetAxis() const
+	EActorModifierAxis GetOrientationAxis() const
 	{
-		return Axis;
+		return OrientationAxis;
 	}
 
 	/** Sets the look-at direction to be flipped. */
@@ -80,19 +80,24 @@ protected:
 
 	UPROPERTY(EditInstanceOnly, Setter="SetReferenceActor", Getter="GetReferenceActor", Category="LookAt", meta=(ShowOnlyInnerProperties, AllowPrivateAccess="true"))
 	FActorModifierSceneTreeActor ReferenceActor;
+	
+	/** The axis to orient look at */
+	UPROPERTY(EditInstanceOnly, Setter, Getter, Category="LookAt", meta=(AllowPrivateAccess="true"))
+	EActorModifierAxis OrientationAxis = EActorModifierAxis::None;
+
+	/** If true, will flip the look-at direction. */
+	UPROPERTY(EditInstanceOnly, Setter="SetFlipAxis", Getter="GetFlipAxis", Category="LookAt", meta=(AllowPrivateAccess="true"))
+	bool bFlipAxis = false;
 
 	/** The actor to look at. */
 	UPROPERTY(meta=(DeprecatedProperty, DeprecationMessage="Use ReferenceActor instead"))
 	TWeakObjectPtr<AActor> ReferenceActorWeak;
-
+	
 	/** The axis that will point towards the reference actor. */
-	UPROPERTY(EditInstanceOnly, Setter="SetAxis", Getter="GetAxis", Category="LookAt", meta=(AllowPrivateAccess="true"))
+	UE_DEPRECATED(5.6, "Use OrientationAxis instead")
+	UPROPERTY(meta=(DeprecatedProperty, DeprecationMessage="Use OrientationAxis instead"))
 	EActorModifierAlignment Axis;
 
-	/** If true, will flip the look-at direction. */
-	UPROPERTY(EditInstanceOnly, Setter="SetFlipAxis", Getter="GetFlipAxis", Category="LookAt", meta=(AllowPrivateAccess="true"))
-	bool bFlipAxis = false;
-
 	UPROPERTY()
 	bool bDeprecatedPropertiesMigrated = false;
 };
diff --git a/Engine/Plugins/Experimental/Avalanche/Source/AvalancheModifiers/Private/Tests/AvaModifiersLookAt.spec.cpp b/Engine/Plugins/Experimental/Avalanche/Source/AvalancheModifiers/Private/Tests/AvaModifiersLookAt.spec.cpp
--- a/Engine/Plugins/Experimental/Avalanche/Source/AvalancheModifiers/Private/Tests/AvaModifiersLookAt.spec.cpp
+++ b/Engine/Plugins/Experimental/Avalanche/Source/AvalancheModifiers/Private/Tests/AvaModifiersLookAt.spec.cpp
@@ -68,7 +68,7 @@ void AvalancheModifiersLookAt::Define()
 			const FVector CurrentModifiedLocation = ModifiedActor->GetActorLocation();
 			const FVector CurrentReferenceLocation = ReferenceActor->GetActorLocation();
 			const FRotator ExpectedRotator = UKismetMathLibrary::FindLookAtRotation(
-				CurrentReferenceLocation, CurrentModifiedLocation);
+				CurrentModifiedLocation, CurrentReferenceLocation);
 
 			TestEqual("Modified actor didn't change location",
 			          CurrentModifiedLocation, InitialModifiedState.GetLocation());
