MotionDesign : Modifiers
- Fixed spline path orientation issue with rotation due to gimbal lock with rotators

diff --git a/Engine/Plugins/Experimental/ActorModifier/Source/ActorModifierLayout/Private/Modifiers/ActorModifierSplinePathModifier.cpp b/Engine/Plugins/Experimental/ActorModifier/Source/ActorModifierLayout/Private/Modifiers/ActorModifierSplinePathModifier.cpp
--- a/Engine/Plugins/Experimental/ActorModifier/Source/ActorModifierLayout/Private/Modifiers/ActorModifierSplinePathModifier.cpp
+++ b/Engine/Plugins/Experimental/ActorModifier/Source/ActorModifierLayout/Private/Modifiers/ActorModifierSplinePathModifier.cpp
@@ -271,8 +271,8 @@ void UActorModifierSplinePathModifier::Apply()
 	if (bOrient)
 	{
 		LayoutShared->SaveActorState(this, ActorModified, EActorModifierTransformSharedState::Rotation);
-		const FRotator WorldRotation = BaseOrientation + SplineComponent->GetRotationAtDistanceAlongSpline(SampleDistance, ESplineCoordinateSpace::World);
-		ActorModified->SetActorRotation(WorldRotation);
+		const FQuat WorldRotationQuat = SplineComponent->GetQuaternionAtDistanceAlongSpline(SampleDistance, ESplineCoordinateSpace::World) * BaseOrientation.Quaternion();
+		ActorModified->SetActorRotation(WorldRotationQuat);
 	}
 	else
 	{
