MotionDesign : Modifiers
- Fix reference container should disable actor picker when value is different than Other

diff --git a/Engine/Plugins/Experimental/ActorModifier/Source/ActorModifier/Public/Extensions/ActorModifierSceneTreeUpdateExtension.h b/Engine/Plugins/Experimental/ActorModifier/Source/ActorModifier/Public/Extensions/ActorModifierSceneTreeUpdateExtension.h
--- a/Engine/Plugins/Experimental/ActorModifier/Source/ActorModifier/Public/Extensions/ActorModifierSceneTreeUpdateExtension.h
+++ b/Engine/Plugins/Experimental/ActorModifier/Source/ActorModifier/Public/Extensions/ActorModifierSceneTreeUpdateExtension.h
@@ -44,11 +44,11 @@ struct FActorModifierSceneTreeActor
 	EActorModifierReferenceContainer ReferenceContainer = EActorModifierReferenceContainer::Other;
 
 	/** The actor being followed by the modifier. This is user selectable if the Reference Container is set to "Other" */
-	UPROPERTY(EditAnywhere, BlueprintReadWrite, Category="Outliner", meta=(DisplayName="Reference Actor", EditCondition = "ReferenceContainer == EAvaReferenceContainer::Other"))
+	UPROPERTY(EditAnywhere, BlueprintReadWrite, Category="Outliner", meta=(DisplayName="Reference Actor", EditCondition = "ReferenceContainer == EActorModifierReferenceContainer::Other"))
 	TWeakObjectPtr<AActor> ReferenceActorWeak = nullptr;
 
 	/** If true, will search for the next visible actor based on the selected reference container */
-	UPROPERTY(EditAnywhere, BlueprintReadWrite, Category="Outliner", meta=(EditCondition="ReferenceContainer != EAvaReferenceContainer::Other", EditConditionHides))
+	UPROPERTY(EditAnywhere, BlueprintReadWrite, Category="Outliner", meta=(EditCondition="ReferenceContainer != EActorModifierReferenceContainer::Other", EditConditionHides))
 	bool bSkipHiddenActors = false;
 
 	bool operator==(const FActorModifierSceneTreeActor& InOther) const
