MotionDesign : Modifiers
- Added FilterActorByComponentClass specifier for actor property to be reused across modifiers to filter level actors based on a component class
- Fixed spline sweep modifier clearing mesh when out of range or invalid
- Allow spline path modifier to loop even when spline is not a closed loop (progress, distance, time, points)

diff --git a/Engine/Plugins/VirtualProduction/ActorModifier/Source/ActorModifierEditor/Private/ActorModifierEditorModule.cpp b/Engine/Plugins/VirtualProduction/ActorModifier/Source/ActorModifierEditor/Private/ActorModifierEditorModule.cpp
--- a/Engine/Plugins/VirtualProduction/ActorModifier/Source/ActorModifierEditor/Private/ActorModifierEditorModule.cpp
+++ b/Engine/Plugins/VirtualProduction/ActorModifier/Source/ActorModifierEditor/Private/ActorModifierEditorModule.cpp
@@ -3,9 +3,8 @@
 #include "ActorModifierEditorModule.h"
 
 #include "ActorModifierTypes.h"
+#include "Customizations/ActorModifierEditorActorComponentClassPropertyTypeCustomization.h"
 #include "Customizations/ActorModifierEditorAnchorAlignmentPropertyTypeCustomization.h"
-#include "Customizations/ActorModifierEditorSplinePathModifierDetailCustomization.h"
-#include "Modifiers/ActorModifierSplinePathModifier.h"
 #include "PropertyEditorModule.h"
 #include "Styles/ActorModifierEditorStyle.h"
 
@@ -13,22 +12,22 @@ IMPLEMENT_MODULE(FActorModifierEditorModule, ActorModifierEditor)
 
 void FActorModifierEditorModule::StartupModule()
 {
-	FActorModifierEditorStyle::Get();
-	
-	FPropertyEditorModule& PropertyModule = FModuleManager::GetModuleChecked<FPropertyEditorModule>(TEXT("PropertyEditor"));
+	using namespace UE::ActorModifierEditor::Private;
 
-	PropertyModule.RegisterCustomPropertyTypeLayout(FActorModifierAnchorAlignment::StaticStruct()->GetFName()
-	, FOnGetPropertyTypeCustomizationInstance::CreateStatic(&FActorModifierEditorAnchorAlignmentPropertyTypeCustomization::MakeInstance));
+	FActorModifierEditorStyle::Get();
 
-	PropertyModule.RegisterCustomClassLayout(UActorModifierSplinePathModifier::StaticClass()->GetFName(), FOnGetDetailCustomizationInstance::CreateStatic(&FActorModifierEditorSplinePathModifierDetailCustomization::MakeInstance));
+	FPropertyEditorModule& PropertyModule = FModuleManager::GetModuleChecked<FPropertyEditorModule>(TEXT("PropertyEditor"));
+	PropertyModule.RegisterCustomPropertyTypeLayout(RegisteredCustomizations.Add_GetRef(FActorModifierAnchorAlignment::StaticStruct()->GetFName()), FOnGetPropertyTypeCustomizationInstance::CreateStatic(&FActorModifierEditorAnchorAlignmentPropertyTypeCustomization::MakeInstance));
+	PropertyModule.RegisterCustomPropertyTypeLayout(RegisteredCustomizations.Add_GetRef(FWeakObjectProperty::StaticClass()->GetFName()), FOnGetPropertyTypeCustomizationInstance::CreateStatic(&FActorModifierEditorActorComponentClassPropertyTypeCustomization::MakeInstance), MakeShared<FActorModifierEditorActorComponentClassPropertyTypeIdentifier>());
 }
 
 void FActorModifierEditorModule::ShutdownModule()
 {
 	if (FPropertyEditorModule* PropertyModule = FModuleManager::GetModulePtr<FPropertyEditorModule>(TEXT("PropertyEditor")))
 	{
-		PropertyModule->UnregisterCustomPropertyTypeLayout(FActorModifierAnchorAlignment::StaticStruct()->GetFName());
-		PropertyModule->UnregisterCustomClassLayout(UActorModifierSplinePathModifier::StaticClass()->GetFName());
+		for (const FName& Name : RegisteredCustomizations)
+		{
+			PropertyModule->UnregisterCustomPropertyTypeLayout(Name);
+		}
 	}
 }
-    
diff --git a/Engine/Plugins/VirtualProduction/ActorModifier/Source/ActorModifierEditor/Private/ActorModifierEditorModule.h b/Engine/Plugins/VirtualProduction/ActorModifier/Source/ActorModifierEditor/Private/ActorModifierEditorModule.h
--- a/Engine/Plugins/VirtualProduction/ActorModifier/Source/ActorModifierEditor/Private/ActorModifierEditorModule.h
+++ b/Engine/Plugins/VirtualProduction/ActorModifier/Source/ActorModifierEditor/Private/ActorModifierEditorModule.h
@@ -2,6 +2,7 @@
 
 #pragma once
 
+#include "Containers/Array.h"
 #include "Modules/ModuleManager.h"
 
 class FActorModifierEditorModule : public IModuleInterface
@@ -11,4 +12,6 @@ public:
     virtual void StartupModule() override;
     virtual void ShutdownModule() override;
 	//~ End IModuleInterface
+private:
+	TArray<FName> RegisteredCustomizations;
 };
diff --git a/Engine/Plugins/VirtualProduction/ActorModifier/Source/ActorModifierEditor/Private/Customizations/ActorModifierEditorActorComponentClassPropertyTypeCustomization.cpp b/Engine/Plugins/VirtualProduction/ActorModifier/Source/ActorModifierEditor/Private/Customizations/ActorModifierEditorActorComponentClassPropertyTypeCustomization.cpp
new file mode 100644
--- /dev/null
+++ b/Engine/Plugins/VirtualProduction/ActorModifier/Source/ActorModifierEditor/Private/Customizations/ActorModifierEditorActorComponentClassPropertyTypeCustomization.cpp
@@ -0,0 +1,60 @@
+// Copyright Epic Games, Inc. All Rights Reserved.
+
+#include "Customizations/ActorModifierEditorActorComponentClassPropertyTypeCustomization.h"
+
+#include "Cloner/Customizations/CEEditorClonerCustomActorPickerNodeBuilder.h"
+#include "Components/ActorComponent.h"
+#include "DetailWidgetRow.h"
+#include "GameFramework/Actor.h"
+
+namespace UE::ActorModifierEditor::Private
+{
+	const FLazyName FActorModifierEditorActorComponentClassPropertyTypeIdentifier::MetadataSpecifier = TEXT("FilterActorByComponentClass");
+
+	bool FActorModifierEditorActorComponentClassPropertyTypeIdentifier::IsPropertyTypeCustomized(const IPropertyHandle& InPropertyHandle) const
+	{
+		if (InPropertyHandle.HasMetaData(MetadataSpecifier))
+		{
+			if (const FWeakObjectProperty* ObjectProperty = CastField<FWeakObjectProperty>(InPropertyHandle.GetProperty()))
+			{
+				return ObjectProperty->PropertyClass && ObjectProperty->PropertyClass->IsChildOf(AActor::StaticClass());
+			}
+		}
+
+		return false;
+	}
+
+	TSharedRef<IPropertyTypeCustomization> FActorModifierEditorActorComponentClassPropertyTypeCustomization::MakeInstance()
+	{
+		return MakeShared<FActorModifierEditorActorComponentClassPropertyTypeCustomization>();
+	}
+
+	void FActorModifierEditorActorComponentClassPropertyTypeCustomization::CustomizeHeader(TSharedRef<IPropertyHandle> InPropertyHandle, FDetailWidgetRow& InHeaderRow, IPropertyTypeCustomizationUtils& InUtils)
+	{
+		InPropertyHandle->MarkHiddenByCustomization();
+
+		ComponentClass = InPropertyHandle->GetClassMetaData(FActorModifierEditorActorComponentClassPropertyTypeIdentifier::MetadataSpecifier);
+
+		ActorPickerNodeBuilder = MakeShared<FCEEditorClonerCustomActorPickerNodeBuilder>(
+			InPropertyHandle
+			, FOnShouldFilterActor::CreateSP(this, &FActorModifierEditorActorComponentClassPropertyTypeCustomization::OnFilterActorByComponentClass)
+		);
+
+		ActorPickerNodeBuilder->GenerateHeaderRowContent(InHeaderRow);
+	}
+
+	void FActorModifierEditorActorComponentClassPropertyTypeCustomization::CustomizeChildren(TSharedRef<IPropertyHandle> InPropertyHandle, IDetailChildrenBuilder& InBuilder, IPropertyTypeCustomizationUtils& InUtils)
+	{
+		ActorPickerNodeBuilder->GenerateChildContent(InBuilder);
+	}
+
+	bool FActorModifierEditorActorComponentClassPropertyTypeCustomization::OnFilterActorByComponentClass(const AActor* InActor)
+	{
+		if (ComponentClass.Get())
+		{
+			return !!InActor->FindComponentByClass(ComponentClass);
+		}
+
+		return true;
+	}
+}
\ No newline at end of file
diff --git a/Engine/Plugins/VirtualProduction/ActorModifier/Source/ActorModifierEditor/Private/Customizations/ActorModifierEditorActorComponentClassPropertyTypeCustomization.h b/Engine/Plugins/VirtualProduction/ActorModifier/Source/ActorModifierEditor/Private/Customizations/ActorModifierEditorActorComponentClassPropertyTypeCustomization.h
new file mode 100644
--- /dev/null
+++ b/Engine/Plugins/VirtualProduction/ActorModifier/Source/ActorModifierEditor/Private/Customizations/ActorModifierEditorActorComponentClassPropertyTypeCustomization.h
@@ -0,0 +1,49 @@
+// Copyright Epic Games, Inc. All Rights Reserved.
+
+#pragma once
+
+#include "IPropertyTypeCustomization.h"
+#include "PropertyEditorModule.h"
+#include "Templates/SubclassOf.h"
+
+class FCEEditorClonerCustomActorPickerNodeBuilder;
+class FDetailWidgetRow;
+class IDetailChildrenBuilder;
+class IPropertyHandle;
+class UActorComponent;
+
+namespace UE::ActorModifierEditor::Private
+{
+	/**
+	 * Only allow property customization with class metadata specifier
+	 */
+	class FActorModifierEditorActorComponentClassPropertyTypeIdentifier : public IPropertyTypeIdentifier
+	{
+	public:
+		static const FLazyName MetadataSpecifier;
+
+		//~ Begin IPropertyTypeIdentifier
+		virtual bool IsPropertyTypeCustomized(const IPropertyHandle& InPropertyHandle) const override;
+		//~ End IPropertyTypeIdentifier
+	};
+
+	/*
+	 * Customization for an actor property to only show filtered actor based on component class
+	*/
+	class FActorModifierEditorActorComponentClassPropertyTypeCustomization : public IPropertyTypeCustomization
+	{
+	public:
+		static TSharedRef<IPropertyTypeCustomization> MakeInstance();
+
+		//~ Begin IPropertyTypeCustomization
+		virtual void CustomizeHeader(TSharedRef<IPropertyHandle> InPropertyHandle, FDetailWidgetRow& InHeaderRow, IPropertyTypeCustomizationUtils& InUtils) override;
+		virtual void CustomizeChildren(TSharedRef<IPropertyHandle> InPropertyHandle, IDetailChildrenBuilder& InBuilder, IPropertyTypeCustomizationUtils& InUtils) override;
+		//~ End IPropertyTypeCustomization
+
+	protected:
+		bool OnFilterActorByComponentClass(const AActor* InActor);
+
+		TSubclassOf<UActorComponent> ComponentClass;
+		TSharedPtr<FCEEditorClonerCustomActorPickerNodeBuilder> ActorPickerNodeBuilder;
+	};
+}
\ No newline at end of file
diff --git a/Engine/Plugins/VirtualProduction/ActorModifier/Source/ActorModifierEditor/Private/Customizations/ActorModifierEditorSplinePathModifierDetailCustomization.cpp b/Engine/Plugins/VirtualProduction/ActorModifier/Source/ActorModifierEditor/Private/Customizations/ActorModifierEditorSplinePathModifierDetailCustomization.cpp
deleted file mode 100644
diff --git a/Engine/Plugins/VirtualProduction/ActorModifier/Source/ActorModifierEditor/Private/Customizations/ActorModifierEditorSplinePathModifierDetailCustomization.h b/Engine/Plugins/VirtualProduction/ActorModifier/Source/ActorModifierEditor/Private/Customizations/ActorModifierEditorSplinePathModifierDetailCustomization.h
deleted file mode 100644
diff --git a/Engine/Plugins/VirtualProduction/ActorModifier/Source/ActorModifierLayout/Private/Modifiers/ActorModifierSplinePathModifier.cpp b/Engine/Plugins/VirtualProduction/ActorModifier/Source/ActorModifierLayout/Private/Modifiers/ActorModifierSplinePathModifier.cpp
--- a/Engine/Plugins/VirtualProduction/ActorModifier/Source/ActorModifierLayout/Private/Modifiers/ActorModifierSplinePathModifier.cpp
+++ b/Engine/Plugins/VirtualProduction/ActorModifier/Source/ActorModifierLayout/Private/Modifiers/ActorModifierSplinePathModifier.cpp
@@ -228,25 +228,50 @@ void UActorModifierSplinePathModifier::Apply()
 
 	if (SampleMode == EActorModifierLayoutSplinePathSampleMode::Distance)
 	{
-		SampleDistance = Distance;
+		if (Distance < 0)
+		{
+			SampleDistance = SplineComponent->GetSplineLength() - FMath::Abs(FMath::Fmod(Distance, SplineComponent->GetSplineLength()));
+		}
+		else
+		{
+			SampleDistance = FMath::Fmod(Distance, SplineComponent->GetSplineLength());
+		}
 	}
 	else if (SampleMode == EActorModifierLayoutSplinePathSampleMode::Percentage)
 	{
-		SampleDistance = SplineComponent->GetSplineLength() * Progress / 100.f;
+		if (Progress < 0)
+		{
+			SampleDistance = SplineComponent->GetSplineLength() - FMath::Abs(FMath::Fmod(SplineComponent->GetSplineLength() * Progress / 100.f, SplineComponent->GetSplineLength()));
+		}
+		else
+		{
+			SampleDistance = FMath::Fmod(SplineComponent->GetSplineLength() * Progress / 100.f, SplineComponent->GetSplineLength());
+		}
 	}
 	else if (SampleMode == EActorModifierLayoutSplinePathSampleMode::Time)
 	{
-		SampleDistance = SplineComponent->GetSplineLength() * Time / SplineComponent->Duration;
+		if (Time < 0)
+		{
+			SampleDistance = SplineComponent->GetSplineLength() - FMath::Abs(FMath::Fmod(SplineComponent->GetSplineLength() * Time / SplineComponent->Duration, SplineComponent->GetSplineLength()));
+		}
+		else
+		{
+			SampleDistance = FMath::Fmod(SplineComponent->GetSplineLength() * Time / SplineComponent->Duration, SplineComponent->GetSplineLength());
+		}
 	}
 	else if (SampleMode == EActorModifierLayoutSplinePathSampleMode::Point)
 	{
 		int32 SampleIndex = PointIndex;
-		
-		if (SplineComponent->IsClosedLoop())
+
+		if (SampleIndex < 0)
+		{
+			SampleIndex = SplineComponent->GetNumberOfSplinePoints() - FMath::Modulo(FMath::Abs(SampleIndex), SplineComponent->GetNumberOfSplinePoints() + 1);
+		}
+		else
 		{
 			SampleIndex = FMath::Modulo(SampleIndex, SplineComponent->GetNumberOfSplinePoints());
 		}
-		
+
 		if (SampleIndex < SplineComponent->GetNumberOfSplinePoints())
 		{
 			SampleDistance = SplineComponent->GetDistanceAlongSplineAtSplinePoint(SampleIndex);
diff --git a/Engine/Plugins/VirtualProduction/ActorModifier/Source/ActorModifierLayout/Public/Modifiers/ActorModifierSplinePathModifier.h b/Engine/Plugins/VirtualProduction/ActorModifier/Source/ActorModifierLayout/Public/Modifiers/ActorModifierSplinePathModifier.h
--- a/Engine/Plugins/VirtualProduction/ActorModifier/Source/ActorModifierLayout/Public/Modifiers/ActorModifierSplinePathModifier.h
+++ b/Engine/Plugins/VirtualProduction/ActorModifier/Source/ActorModifierLayout/Public/Modifiers/ActorModifierSplinePathModifier.h
@@ -148,7 +148,7 @@ protected:
 	void OnSplineOptionsChanged();
 	
 	/** Spline actor to retrieve the USplineComponent from */
-	UPROPERTY(EditInstanceOnly, Setter, Getter, Category="SplinePath", meta=(AllowPrivateAccess="true"))
+	UPROPERTY(EditInstanceOnly, Setter, Getter, Category="SplinePath", meta=(FilterActorByComponentClass="/Script/Engine.SplineComponent", AllowPrivateAccess="true"))
 	TWeakObjectPtr<AActor> SplineActorWeak;
 
 	/** How to sample the spline */
@@ -156,19 +156,19 @@ protected:
 	EActorModifierLayoutSplinePathSampleMode SampleMode = EActorModifierLayoutSplinePathSampleMode::Percentage;
 	
 	/** Percentage progress to sample the spline at */
-	UPROPERTY(EditInstanceOnly, Setter, Getter, Category="SplinePath", meta=(EditCondition="SampleMode == EActorModifierLayoutSplinePathSampleMode::Percentage", EditConditionHides, Units=Percent, ClampMin="0", UIMin="0", UIMax="100", AllowPrivateAccess="true"))
+	UPROPERTY(EditInstanceOnly, Setter, Getter, Category="SplinePath", meta=(EditCondition="SampleMode == EActorModifierLayoutSplinePathSampleMode::Percentage", EditConditionHides, Units=Percent, AllowPrivateAccess="true"))
 	float Progress = 0.f;
 
 	/** Distance to sample the spline at */
-	UPROPERTY(EditInstanceOnly, Setter, Getter, Category="SplinePath", meta=(EditCondition="SampleMode == EActorModifierLayoutSplinePathSampleMode::Distance", EditConditionHides, ClampMin="0", Units=Centimeters, AllowPrivateAccess="true"))
+	UPROPERTY(EditInstanceOnly, Setter, Getter, Category="SplinePath", meta=(EditCondition="SampleMode == EActorModifierLayoutSplinePathSampleMode::Distance", EditConditionHides, Units=Centimeters, AllowPrivateAccess="true"))
 	float Distance = 0.f;
 
 	/** Time to sample the spline at */
-	UPROPERTY(EditInstanceOnly, Setter, Getter, Category="SplinePath", meta=(EditCondition="SampleMode == EActorModifierLayoutSplinePathSampleMode::Time", EditConditionHides, ClampMin="0", AllowPrivateAccess="true"))
+	UPROPERTY(EditInstanceOnly, Setter, Getter, Category="SplinePath", meta=(EditCondition="SampleMode == EActorModifierLayoutSplinePathSampleMode::Time", EditConditionHides, AllowPrivateAccess="true"))
 	float Time = 0.f;
 	
 	/** Point index to sample the spline at */
-	UPROPERTY(EditInstanceOnly, Setter, Getter, Category="SplinePath", meta=(EditCondition="SampleMode == EActorModifierLayoutSplinePathSampleMode::Point", EditConditionHides, ClampMin="0", AllowPrivateAccess="true"))
+	UPROPERTY(EditInstanceOnly, Setter, Getter, Category="SplinePath", meta=(EditCondition="SampleMode == EActorModifierLayoutSplinePathSampleMode::Point", EditConditionHides, AllowPrivateAccess="true"))
 	int32 PointIndex = 0;
 
 	/** Orient actor based on spline tangent */
diff --git a/Engine/Plugins/VirtualProduction/Avalanche/Source/AvalancheModifiers/Private/Modifiers/AvaSplineSweepModifier.cpp b/Engine/Plugins/VirtualProduction/Avalanche/Source/AvalancheModifiers/Private/Modifiers/AvaSplineSweepModifier.cpp
--- a/Engine/Plugins/VirtualProduction/Avalanche/Source/AvalancheModifiers/Private/Modifiers/AvaSplineSweepModifier.cpp
+++ b/Engine/Plugins/VirtualProduction/Avalanche/Source/AvalancheModifiers/Private/Modifiers/AvaSplineSweepModifier.cpp
@@ -38,16 +38,12 @@ void UAvaSplineSweepModifier::PostEditChangeProperty(FPropertyChangedEvent& InPr
 
 void UAvaSplineSweepModifier::SetSteps(int32 InSteps)
 {
+	InSteps = FMath::Clamp(InSteps, 3, MaxSampleCount);
 	if (Steps == InSteps)
 	{
 		return;
 	}
 
-	if (InSteps < 0 || InSteps > MaxSampleCount)
-	{
-		return;
-	}
-
 	Steps = InSteps;
 	OnSplineOptionsChanged();
 }
@@ -59,11 +55,6 @@ void UAvaSplineSweepModifier::SetProgressOffset(float InOffset)
 		return;
 	}
 
-	if (InOffset < 0.f)
-	{
-		return;
-	}
-
 	ProgressOffset = InOffset;
 	OnSplineOptionsChanged();
 }
@@ -207,6 +198,7 @@ void UAvaSplineSweepModifier::OnModifierCDOSetup(FActorModifierCoreMetadata& InM
 	InMetadata.SetName(TEXT("SplineSweep"));
 	InMetadata.SetCategory(TEXT("Geometry"));
 #if WITH_EDITOR
+	InMetadata.SetDisplayName(LOCTEXT("ModifierDisplayName", "Spline Sweep"));
 	InMetadata.SetDescription(LOCTEXT("ModifierDescription", "Samples a spline path and creates a 3D geometry based on a 2D shape slice"));
 #endif
 
@@ -253,12 +245,32 @@ void UAvaSplineSweepModifier::Apply()
 		return;
 	}
 
+	const bool bInvalidProgress = ProgressStart > ProgressEnd;
+	const bool bOutOfRange = !SplineComponent->IsClosedLoop()
+		&& (ProgressStart + ProgressOffset > 1.f || ProgressEnd + ProgressOffset < 0.f);
+
+	if (bOutOfRange || bInvalidProgress)
+	{
+		// Clear mesh
+		TargetMesh->EditMesh([](FDynamicMesh3& InMesh)
+		{
+			for (const int32 TId : InMesh.TriangleIndicesItr())
+			{
+				InMesh.RemoveTriangle(TId);
+			}
+		});
+
+		Next();
+		return;
+	}
+
 	if (Steps != SplineSamples.Num())
 	{
 		SampleSpline();
 	}
 
-	if (SplineSamples.IsEmpty())
+	// Samples needed should be >= 3 to avoid any issue during AppendSweepPolygon operation
+	if (SplineSamples.Num() < 3)
 	{
 		Next();
 		return;
diff --git a/Engine/Plugins/VirtualProduction/Avalanche/Source/AvalancheModifiers/Public/Modifiers/AvaSplineSweepModifier.h b/Engine/Plugins/VirtualProduction/Avalanche/Source/AvalancheModifiers/Public/Modifiers/AvaSplineSweepModifier.h
--- a/Engine/Plugins/VirtualProduction/Avalanche/Source/AvalancheModifiers/Public/Modifiers/AvaSplineSweepModifier.h
+++ b/Engine/Plugins/VirtualProduction/Avalanche/Source/AvalancheModifiers/Public/Modifiers/AvaSplineSweepModifier.h
@@ -25,7 +25,7 @@ class UAvaSplineSweepModifier : public UAvaGeometryBaseModifier
 	GENERATED_BODY()
 
 public:
-	static constexpr int32 MaxSampleCount = 100;
+	static constexpr int32 MaxSampleCount = 200;
 
 	UFUNCTION(BlueprintCallable, Category="Motion Design|Modifiers|SplineSweep")
 	AVALANCHEMODIFIERS_API void SetSplineActor(AActor* InActor)
@@ -159,7 +159,7 @@ protected:
 	bool SampleSpline();
 
 	/** Spline actor to retrieve the USplineComponent from */
-	UPROPERTY(EditInstanceOnly, Setter="SetSplineActorWeak", Getter="GetSplineActorWeak", Category="SplineSweep", meta=(AllowPrivateAccess="true"))
+	UPROPERTY(EditInstanceOnly, Setter="SetSplineActorWeak", Getter="GetSplineActorWeak", Category="SplineSweep", meta=(FilterActorByComponentClass="/Script/Engine.SplineComponent", AllowPrivateAccess="true"))
 	TWeakObjectPtr<AActor> SplineActorWeak;
 
 	/** How do we sample the spline */
@@ -171,11 +171,11 @@ protected:
 	float SampleDistance = 1000.f;
 
 	/** The sample count defines the precision of the sweep */
-	UPROPERTY(EditInstanceOnly, Setter="SetSteps", Getter="GetSteps", Category="SplineSweep", meta=(ClampMin="0", ClampMax="100", AllowPrivateAccess="true"))
+	UPROPERTY(EditInstanceOnly, Setter="SetSteps", Getter="GetSteps", Category="SplineSweep", meta=(ClampMin="3", ClampMax="200", AllowPrivateAccess="true"))
 	int32 Steps = 10;
 
 	/** Sample range offset of the spline, for closed loop spline range can wrap around */
-	UPROPERTY(EditInstanceOnly, Setter="SetProgressOffset", Getter="GetProgressOffset", Category="SplineSweep", meta=(ClampMin="0", AllowPrivateAccess="true"))
+	UPROPERTY(EditInstanceOnly, Setter="SetProgressOffset", Getter="GetProgressOffset", Category="SplineSweep", meta=(AllowPrivateAccess="true"))
 	float ProgressOffset = 0.f;
 
 	/** Sample start range of the spline */
