MotionDesign : ActorModifier
- Added an holdout composite modifier that acts like the holdout composite component and improves primitives output when it is animated by rendering them in a separate pass then compositing it back

diff --git a/Engine/Plugins/VirtualProduction/ActorModifier/ActorModifier.uplugin b/Engine/Plugins/VirtualProduction/ActorModifier/ActorModifier.uplugin
--- a/Engine/Plugins/VirtualProduction/ActorModifier/ActorModifier.uplugin
+++ b/Engine/Plugins/VirtualProduction/ActorModifier/ActorModifier.uplugin
@@ -34,6 +34,11 @@
 			"Name": "ActorModifierLayout",
 			"Type": "Runtime",
 			"LoadingPhase": "Default"
+		},
+		{
+			"Name": "ActorModifierRendering",
+			"Type": "Runtime",
+			"LoadingPhase": "Default"
 		}
 	],
 	"Plugins": [
@@ -44,6 +49,10 @@
 		{
 			"Name": "ClonerEffector",
 			"Enabled": true
+		},
+		{
+			"Name": "CompositeCore",
+			"Enabled": true
 		}
 	]
 }
\ No newline at end of file
diff --git a/Engine/Plugins/VirtualProduction/ActorModifier/Source/ActorModifierLayout/ActorModifierLayout.Build.cs b/Engine/Plugins/VirtualProduction/ActorModifier/Source/ActorModifierLayout/ActorModifierLayout.Build.cs
--- a/Engine/Plugins/VirtualProduction/ActorModifier/Source/ActorModifierLayout/ActorModifierLayout.Build.cs
+++ b/Engine/Plugins/VirtualProduction/ActorModifier/Source/ActorModifierLayout/ActorModifierLayout.Build.cs
@@ -21,9 +21,7 @@ public class ActorModifierLayout : ModuleRules
             new string[]
             {
                 "CoreUObject",
-                "Engine",
-                "Slate",
-                "SlateCore"
+                "Engine"
             }
         );
     }
diff --git a/Engine/Plugins/VirtualProduction/ActorModifier/Source/ActorModifierRendering/ActorModifierRendering.Build.cs b/Engine/Plugins/VirtualProduction/ActorModifier/Source/ActorModifierRendering/ActorModifierRendering.Build.cs
new file mode 100644
--- /dev/null
+++ b/Engine/Plugins/VirtualProduction/ActorModifier/Source/ActorModifierRendering/ActorModifierRendering.Build.cs
@@ -0,0 +1,29 @@
+﻿// Copyright Epic Games, Inc. All Rights Reserved.
+
+using UnrealBuildTool;
+
+public class ActorModifierRendering : ModuleRules
+{
+    public ActorModifierRendering(ReadOnlyTargetRules Target) : base(Target)
+    {
+        PCHUsage = ModuleRules.PCHUsageMode.UseExplicitOrSharedPCHs;
+
+        PublicDependencyModuleNames.AddRange(
+            new string[]
+            {
+	            "ActorModifier",
+	            "ActorModifierCore",
+                "Core",
+            }
+        );
+
+        PrivateDependencyModuleNames.AddRange(
+            new string[]
+            {
+	            "CompositeCore",
+                "CoreUObject",
+                "Engine"
+            }
+        );
+    }
+}
\ No newline at end of file
diff --git a/Engine/Plugins/VirtualProduction/ActorModifier/Source/ActorModifierRendering/Private/ActorModifierRenderingModule.cpp b/Engine/Plugins/VirtualProduction/ActorModifier/Source/ActorModifierRendering/Private/ActorModifierRenderingModule.cpp
new file mode 100644
--- /dev/null
+++ b/Engine/Plugins/VirtualProduction/ActorModifier/Source/ActorModifierRendering/Private/ActorModifierRenderingModule.cpp
@@ -0,0 +1,5 @@
+﻿// Copyright Epic Games, Inc. All Rights Reserved.
+
+#include "Modules/ModuleManager.h"
+    
+IMPLEMENT_MODULE(FDefaultModuleImpl, ActorModifierRendering)
\ No newline at end of file
diff --git a/Engine/Plugins/VirtualProduction/ActorModifier/Source/ActorModifierRendering/Private/Modifiers/ActorModifierHoldoutCompositeModifier.cpp b/Engine/Plugins/VirtualProduction/ActorModifier/Source/ActorModifierRendering/Private/Modifiers/ActorModifierHoldoutCompositeModifier.cpp
new file mode 100644
--- /dev/null
+++ b/Engine/Plugins/VirtualProduction/ActorModifier/Source/ActorModifierRendering/Private/Modifiers/ActorModifierHoldoutCompositeModifier.cpp
@@ -0,0 +1,159 @@
+// Copyright Epic Games, Inc. All Rights Reserved.
+
+#include "Modifiers/ActorModifierHoldoutCompositeModifier.h"
+#include "Components/PrimitiveComponent.h"
+#include "CompositeCoreSubsystem.h"
+#include "HoldoutCompositeComponent.h"
+
+#define LOCTEXT_NAMESPACE "ActorModifierHoldoutModifier"
+
+void UActorModifierHoldoutCompositeModifier::SetIncludeChildren(bool bInInclude)
+{
+	if (bInInclude == bIncludeChildren)
+	{
+		return;
+	}
+
+	bIncludeChildren = bInInclude;
+	OnIncludeChildrenChanged();
+}
+
+#if WITH_EDITOR
+void UActorModifierHoldoutCompositeModifier::PostEditChangeProperty(FPropertyChangedEvent& InPropertyChangedEvent)
+{
+	Super::PostEditChangeProperty(InPropertyChangedEvent);
+
+	const FName MemberPropertyName = InPropertyChangedEvent.GetMemberPropertyName();
+
+	if (MemberPropertyName == GET_MEMBER_NAME_CHECKED(UActorModifierHoldoutCompositeModifier, bIncludeChildren))
+	{
+		OnIncludeChildrenChanged();
+	}
+}
+#endif
+
+void UActorModifierHoldoutCompositeModifier::OnModifierCDOSetup(FActorModifierCoreMetadata& InMetadata)
+{
+	Super::OnModifierCDOSetup(InMetadata);
+
+	InMetadata.SetName(TEXT("HoldoutComposite"));
+	InMetadata.SetCategory(TEXT("Rendering"));
+	InMetadata.SetCompatibilityRule([](const AActor* InActor)->bool
+	{
+		return InActor->FindComponentByClass<UHoldoutCompositeComponent>() == nullptr;
+	});
+#if WITH_EDITOR
+	InMetadata.SetDisplayName(LOCTEXT("ModifierDisplayName", "Holdout composite"));
+	InMetadata.SetDescription(LOCTEXT("ModifierDescription", "Renders the attached primitives in a separate pass and composites it back using an alpha holdout mask"));
+#endif
+}
+
+void UActorModifierHoldoutCompositeModifier::OnModifierAdded(EActorModifierCoreEnableReason InReason)
+{
+	Super::OnModifierAdded(InReason);
+
+	// Track hierarchy changes
+	if (FActorModifierSceneTreeUpdateExtension* SceneExtension = GetExtension<FActorModifierSceneTreeUpdateExtension>())
+	{
+		ReferenceActor.ReferenceContainer = EActorModifierReferenceContainer::Other;
+		ReferenceActor.ReferenceActorWeak = GetModifiedActor();
+		ReferenceActor.bSkipHiddenActors = false;
+		
+		SceneExtension->TrackSceneTree(0, &ReferenceActor);
+	}
+}
+
+void UActorModifierHoldoutCompositeModifier::RestorePreState()
+{
+	Super::RestorePreState();
+
+	UnregisterPrimitiveComponents();
+}
+
+void UActorModifierHoldoutCompositeModifier::Apply()
+{
+	UCompositeCoreSubsystem* CompositeSubsystem = UWorld::GetSubsystem<UCompositeCoreSubsystem>(GetWorld());
+	if (!CompositeSubsystem)
+	{
+		Fail(LOCTEXT("InvalidCompositeSubsystem", "Composite subsystem is not valid"));
+		return;
+	}
+
+	if (!UCompositeCoreSubsystem::IsProjectSettingsValid())
+	{
+		Fail(LOCTEXT("InvalidProjectSettings", "Invalid holdout rendering project settings"));
+		return;
+	}
+
+	// Scan primitive components
+	ForEachComponent<UPrimitiveComponent>(
+		[this](UPrimitiveComponent* InComponent)->bool
+		{
+			PrimitiveComponentsWeak.Add(InComponent);
+			return true;
+		}
+		, EActorModifierCoreComponentType::All
+		, bIncludeChildren ? EActorModifierCoreLookup::SelfAndAllChildren : EActorModifierCoreLookup::Self
+	);
+
+	TArray<UPrimitiveComponent*> PrimitiveComponents;
+	PrimitiveComponents.Reserve(PrimitiveComponentsWeak.Num());
+
+	Algo::TransformIf(
+		PrimitiveComponentsWeak,
+		PrimitiveComponents,
+		[](const TWeakObjectPtr<UPrimitiveComponent>& InComponentWeak)
+		{
+			return InComponentWeak.IsValid();
+		},
+		[](const TWeakObjectPtr<UPrimitiveComponent>& InComponentWeak)
+		{
+			return InComponentWeak.Get();
+		}
+	);
+
+	CompositeSubsystem->RegisterPrimitives(PrimitiveComponents);
+
+	Next();
+}
+
+void UActorModifierHoldoutCompositeModifier::OnSceneTreeTrackedActorChildrenChanged(int32 InIdx, const TSet<TWeakObjectPtr<AActor>>& InPreviousChildrenActors, const TSet<TWeakObjectPtr<AActor>>& InNewChildrenActors)
+{
+	Super::OnSceneTreeTrackedActorChildrenChanged(InIdx, InPreviousChildrenActors, InNewChildrenActors);
+
+	if (bIncludeChildren)
+	{
+		MarkModifierDirty();
+	}
+}
+
+void UActorModifierHoldoutCompositeModifier::OnIncludeChildrenChanged()
+{
+	MarkModifierDirty();
+}
+
+void UActorModifierHoldoutCompositeModifier::UnregisterPrimitiveComponents()
+{
+	if (UCompositeCoreSubsystem* CompositeSubsystem = UWorld::GetSubsystem<UCompositeCoreSubsystem>(GetWorld()))
+	{
+		TArray<UPrimitiveComponent*> PrimitiveComponents;
+		Algo::TransformIf(
+			PrimitiveComponentsWeak,
+			PrimitiveComponents,
+			[](const TWeakObjectPtr<UPrimitiveComponent>& InComponentWeak)
+			{
+				return InComponentWeak.IsValid();
+			},
+			[](const TWeakObjectPtr<UPrimitiveComponent>& InComponentWeak)
+			{
+				return InComponentWeak.Get();
+			}
+		);
+
+		CompositeSubsystem->UnregisterPrimitives(PrimitiveComponents);
+
+		PrimitiveComponentsWeak.Empty();
+	}
+}
+
+#undef LOCTEXT_NAMESPACE
diff --git a/Engine/Plugins/VirtualProduction/ActorModifier/Source/ActorModifierRendering/Public/Modifiers/ActorModifierHoldoutCompositeModifier.h b/Engine/Plugins/VirtualProduction/ActorModifier/Source/ActorModifierRendering/Public/Modifiers/ActorModifierHoldoutCompositeModifier.h
new file mode 100644
--- /dev/null
+++ b/Engine/Plugins/VirtualProduction/ActorModifier/Source/ActorModifierRendering/Public/Modifiers/ActorModifierHoldoutCompositeModifier.h
@@ -0,0 +1,62 @@
+// Copyright Epic Games, Inc. All Rights Reserved.
+
+#pragma once
+
+#include "Containers/Set.h"
+#include "Modifiers/ActorModifierAttachmentBaseModifier.h"
+#include "ActorModifierHoldoutCompositeModifier.generated.h"
+
+class AActor;
+class UPrimitiveComponent;
+
+/**
+ * Track primitive components to render them in a separate pass and composites it back using an alpha holdout mask
+ */
+UCLASS(MinimalAPI, BlueprintType)
+class UActorModifierHoldoutCompositeModifier : public UActorModifierAttachmentBaseModifier
+{
+	GENERATED_BODY()
+
+public:
+	UFUNCTION(BlueprintCallable, Category="Motion Design|Modifiers|Holdout")
+	void SetIncludeChildren(bool bInInclude);
+
+	UFUNCTION(BlueprintCallable, Category="Motion Design|Modifiers|Holdout")
+	bool GetIncludeChildren() const
+	{
+		return bIncludeChildren;
+	}
+
+protected:
+	//~ Begin UObject
+#if WITH_EDITOR
+	virtual void PostEditChangeProperty(FPropertyChangedEvent& InPropertyChangedEvent) override;
+#endif
+	//~ End UObject
+
+	//~ Begin UActorModifierCoreBase
+	virtual void OnModifierCDOSetup(FActorModifierCoreMetadata& InMetadata) override;
+	virtual void OnModifierAdded(EActorModifierCoreEnableReason InReason) override;
+	virtual void RestorePreState() override;
+	virtual void Apply() override;
+	//~ End UActorModifierCoreBase
+
+	//~ Begin IAvaSceneTreeUpdateModifierExtension
+	virtual void OnSceneTreeTrackedActorChildrenChanged(int32 InIdx, const TSet<TWeakObjectPtr<AActor>>& InPreviousChildrenActors, const TSet<TWeakObjectPtr<AActor>>& InNewChildrenActors) override;
+	//~ End IAvaSceneTreeUpdateModifierExtension
+
+	void OnIncludeChildrenChanged();
+	void UnregisterPrimitiveComponents();
+
+	/** If true, will apply itself onto children actors as well */
+	UPROPERTY(EditInstanceOnly, Setter="SetIncludeChildren", Getter="GetIncludeChildren", Category="HoldoutComposite", meta=(AllowPrivateAccess="true"))
+	bool bIncludeChildren = false;
+
+	/** Primitive components tracked by the composite subsystem */
+	UPROPERTY()
+	TSet<TWeakObjectPtr<UPrimitiveComponent>> PrimitiveComponentsWeak;
+
+	/** Used to track self modified actor for changes */
+	UPROPERTY(Transient, DuplicateTransient, TextExportTransient)
+	FActorModifierSceneTreeActor ReferenceActor;
+};
