MotionDesign : ActorModifier
- Fixed justify modifier not working when container is rotated

diff --git a/Engine/Plugins/VirtualProduction/ActorModifier/Source/ActorModifier/Private/Utilities/ActorModifierActorUtils.cpp b/Engine/Plugins/VirtualProduction/ActorModifier/Source/ActorModifier/Private/Utilities/ActorModifierActorUtils.cpp
--- a/Engine/Plugins/VirtualProduction/ActorModifier/Source/ActorModifier/Private/Utilities/ActorModifierActorUtils.cpp
+++ b/Engine/Plugins/VirtualProduction/ActorModifier/Source/ActorModifier/Private/Utilities/ActorModifierActorUtils.cpp
@@ -28,7 +28,7 @@ bool UE::ActorModifier::ActorUtils::IsActorVisible(const AActor* InActor)
 	return bIsActorVisible && bIsActorVisibleInEditor && bIsRootComponentVisible;
 }
 
-FBox UE::ActorModifier::ActorUtils::GetActorsBounds(const TSet<TWeakObjectPtr<AActor>>& InActors, const FTransform& InReferenceTransform, bool bInSkipHidden)
+FBox UE::ActorModifier::ActorUtils::GetActorsBounds(const TSet<TWeakObjectPtr<AActor>>& InActors, const FTransform& InReferenceTransform, bool bInSkipHidden, bool bInTransformBox)
 {
 	FBox ActorsBounds = FBox(EForceInit::ForceInit);
 	ActorsBounds.IsValid = 0;
@@ -51,7 +51,7 @@ FBox UE::ActorModifier::ActorUtils::GetActorsBounds(const TSet<TWeakObjectPtr<AA
 			continue;
 		}
 
-		FBox ActorBounds = UE::ActorModifier::ActorUtils::GetActorBounds(Actor);
+		const FBox ActorBounds = UE::ActorModifier::ActorUtils::GetActorBounds(Actor);
 
 		if (ActorBounds.IsValid > 0)
 		{
@@ -73,7 +73,13 @@ FBox UE::ActorModifier::ActorUtils::GetActorsBounds(const TSet<TWeakObjectPtr<AA
 
 	AccumulatedTransform.SetScale3D(FVector::OneVector);
 
-	return ActorsBounds.TransformBy(AccumulatedTransform);
+	return bInTransformBox ? ActorsBounds.TransformBy(AccumulatedTransform) : ActorsBounds;
+}
+
+FBox UE::ActorModifier::ActorUtils::GetActorsBounds(const TSet<TWeakObjectPtr<AActor>>& InActors, const FTransform& InReferenceTransform, bool bInSkipHidden)
+{
+	constexpr bool bTransformBox = true;
+	return GetActorsBounds(InActors, InReferenceTransform, bInSkipHidden, bTransformBox);
 }
 
 FBox UE::ActorModifier::ActorUtils::GetActorsBounds(AActor* InActor, bool bInIncludeChildren, bool bInSkipHidden)
diff --git a/Engine/Plugins/VirtualProduction/ActorModifier/Source/ActorModifier/Public/Utilities/ActorModifierActorUtils.h b/Engine/Plugins/VirtualProduction/ActorModifier/Source/ActorModifier/Public/Utilities/ActorModifierActorUtils.h
--- a/Engine/Plugins/VirtualProduction/ActorModifier/Source/ActorModifier/Public/Utilities/ActorModifierActorUtils.h
+++ b/Engine/Plugins/VirtualProduction/ActorModifier/Source/ActorModifier/Public/Utilities/ActorModifierActorUtils.h
@@ -17,6 +17,7 @@ namespace UE::ActorModifier::ActorUtils
 	ACTORMODIFIER_API bool IsActorVisible(const AActor* InActor);
 
 	// Layout
+	ACTORMODIFIER_API FBox GetActorsBounds(const TSet<TWeakObjectPtr<AActor>>& InActors, const FTransform& InReferenceTransform, bool bInSkipHidden, bool bInTransformBox);
 	ACTORMODIFIER_API FBox GetActorsBounds(const TSet<TWeakObjectPtr<AActor>>& InActors, const FTransform& InReferenceTransform, bool bInSkipHidden = false);
 	ACTORMODIFIER_API FBox GetActorsBounds(AActor* InActor, bool bInIncludeChildren, bool bInSkipHidden = false);
 	ACTORMODIFIER_API FBox GetActorBounds(const AActor* InActor);
diff --git a/Engine/Plugins/VirtualProduction/ActorModifier/Source/ActorModifierLayout/Private/Modifiers/ActorModifierJustifyModifier.cpp b/Engine/Plugins/VirtualProduction/ActorModifier/Source/ActorModifierLayout/Private/Modifiers/ActorModifierJustifyModifier.cpp
--- a/Engine/Plugins/VirtualProduction/ActorModifier/Source/ActorModifierLayout/Private/Modifiers/ActorModifierJustifyModifier.cpp
+++ b/Engine/Plugins/VirtualProduction/ActorModifier/Source/ActorModifierLayout/Private/Modifiers/ActorModifierJustifyModifier.cpp
@@ -31,13 +31,15 @@ bool UActorModifierJustifyModifier::IsModifierDirtyable() const
 		return Super::IsModifierDirtyable();
 	}
 
-	const FBox TrackedActorLocalBounds = UE::ActorModifier::ActorUtils::GetActorsBounds(ChildrenActorsWeak, FTransform::Identity, true);
+	constexpr bool bSkipHidden = true;
+	constexpr bool bTransformBox = false;
+	const FBox TrackedActorLocalBounds = UE::ActorModifier::ActorUtils::GetActorsBounds(ChildrenActorsWeak, FTransform::Identity, bSkipHidden, bTransformBox);
 
 	if (TrackedActorLocalBounds.Equals(CachedTrackedBounds, 0.01))
 	{
 		return Super::IsModifierDirtyable();
 	}
-	
+
 	return true;
 }
 
@@ -55,18 +57,20 @@ void UActorModifierJustifyModifier::Apply()
 	GetChildrenActors(NewChildrenActorsWeak);
 
 	// the children bounds need to be aligned to the modified actor
-	CachedTrackedBounds = UE::ActorModifier::ActorUtils::GetActorsBounds(NewChildrenActorsWeak, FTransform::Identity, true);
-	
+	constexpr bool bSkipHidden = true;
+	constexpr bool bTransformBox = false;
+	CachedTrackedBounds = UE::ActorModifier::ActorUtils::GetActorsBounds(NewChildrenActorsWeak, ActorModified->GetActorTransform(), bSkipHidden, bTransformBox);
+
 	FVector BoundsCenter;
 	FVector BoundsExtent;
 	CachedTrackedBounds.GetCenterAndExtents(BoundsCenter, BoundsExtent);
 
 	const FVector AlignmentOffset = GetAlignmentOffset(BoundsExtent);
 	const FVector AnchorOffsetVector = GetAnchorOffset();
-	
+
 	// used to store the offset needed to constrain or un-constrain the justification axis, the modified actor and its children movements
 	// constraints start set to the bounds center vector, and are changed below based on different combinations
-	const FVector ConstraintVector = GetConstraintVector(BoundsCenter, ActorModified->GetActorLocation());
+	const FVector ConstraintVector = GetConstraintVector(BoundsCenter, FVector::ZeroVector);
 	const FVector ChildLocationOffset = ConstraintVector + AlignmentOffset - AnchorOffsetVector;
 
 	constexpr bool bCreate = true;
@@ -101,9 +105,11 @@ void UActorModifierJustifyModifier::Apply()
 			}
 
 			LayoutShared->SaveActorState(this, Child);
-			const FVector ChildRelativeLocation = Child->GetRootComponent()->GetRelativeLocation();
-			const FVector DesiredChildLocation = ChildRelativeLocation - ChildLocationOffset;
-			Child->SetActorRelativeLocation(DesiredChildLocation);
+
+			const FVector RelativeLocation = Child->GetRootComponent()->GetRelativeLocation();
+			const FVector DesiredLocation = RelativeLocation - ChildLocationOffset;
+
+			Child->SetActorRelativeLocation(DesiredLocation);
 		}
 	}
 		
@@ -124,19 +130,20 @@ void UActorModifierJustifyModifier::Apply()
 void UActorModifierJustifyModifier::PostEditChangeProperty(FPropertyChangedEvent& InPropertyChangedEvent)
 {
 	Super::PostEditChangeProperty(InPropertyChangedEvent);
-	
+
 	const FName PropertyName = InPropertyChangedEvent.GetMemberPropertyName();
-	
-	static const FName HorizontalAlignmentPropertyName = GET_MEMBER_NAME_CHECKED(UActorModifierJustifyModifier, HorizontalAlignment);
-	static const FName VerticalAlignmentPropertyName = GET_MEMBER_NAME_CHECKED(UActorModifierJustifyModifier, VerticalAlignment);
-	static const FName DepthAlignmentPropertyName = GET_MEMBER_NAME_CHECKED(UActorModifierJustifyModifier, DepthAlignment);
-	static const FName HorizontalAnchorPropertyName = GET_MEMBER_NAME_CHECKED(UActorModifierJustifyModifier, HorizontalAnchor);
-	static const FName VerticalAnchorPropertyName = GET_MEMBER_NAME_CHECKED(UActorModifierJustifyModifier, VerticalAnchor);
-	static const FName DepthAnchorPropertyName = GET_MEMBER_NAME_CHECKED(UActorModifierJustifyModifier, DepthAnchor);
-	
-	if (PropertyName == HorizontalAlignmentPropertyName ||
-		PropertyName == VerticalAlignmentPropertyName ||
-		PropertyName == DepthAlignmentPropertyName)
+
+	static const TSet<FName> PropertyNames =
+	{
+		GET_MEMBER_NAME_CHECKED(UActorModifierJustifyModifier, HorizontalAlignment),
+		GET_MEMBER_NAME_CHECKED(UActorModifierJustifyModifier, VerticalAlignment),
+		GET_MEMBER_NAME_CHECKED(UActorModifierJustifyModifier, DepthAlignment),
+		GET_MEMBER_NAME_CHECKED(UActorModifierJustifyModifier, HorizontalAnchor),
+		GET_MEMBER_NAME_CHECKED(UActorModifierJustifyModifier, VerticalAnchor),
+		GET_MEMBER_NAME_CHECKED(UActorModifierJustifyModifier, DepthAnchor)
+	};
+
+	if (PropertyNames.Contains(PropertyName))
 	{
 		MarkModifierDirty();
 	}
