MotionDesign : ActorModifier
- Fixed justify modifier not working when container is rotated

diff --git a/Engine/Plugins/VirtualProduction/ActorModifier/Source/ActorModifierLayout/Private/Modifiers/ActorModifierJustifyModifier.cpp b/Engine/Plugins/VirtualProduction/ActorModifier/Source/ActorModifierLayout/Private/Modifiers/ActorModifierJustifyModifier.cpp
--- a/Engine/Plugins/VirtualProduction/ActorModifier/Source/ActorModifierLayout/Private/Modifiers/ActorModifierJustifyModifier.cpp
+++ b/Engine/Plugins/VirtualProduction/ActorModifier/Source/ActorModifierLayout/Private/Modifiers/ActorModifierJustifyModifier.cpp
@@ -33,7 +33,7 @@ bool UActorModifierJustifyModifier::IsModifierDirtyable() const
 
 	constexpr bool bSkipHidden = true;
 	constexpr bool bTransformBox = false;
-	const FBox TrackedActorLocalBounds = UE::ActorModifier::ActorUtils::GetActorsBounds(ChildrenActorsWeak, FTransform::Identity, bSkipHidden, bTransformBox);
+	const FBox TrackedActorLocalBounds = UE::ActorModifier::ActorUtils::GetActorsBounds(ChildrenActorsWeak, ActorModified->GetActorTransform(), bSkipHidden, bTransformBox);
 
 	if (TrackedActorLocalBounds.Equals(CachedTrackedBounds, 0.01))
 	{
